- Adding Rivermax Player to add IP video playback in engine by using Rivermax Input Stream
- Adding Rivermax Capture to add video capture over IP by using Rivermax Output stream

diff --git a/Engine/Plugins/Media/AjaMedia/Source/AjaMediaOutput/Private/AjaMediaCapture.cpp b/Engine/Plugins/Media/AjaMedia/Source/AjaMediaOutput/Private/AjaMediaCapture.cpp
--- a/Engine/Plugins/Media/AjaMedia/Source/AjaMediaOutput/Private/AjaMediaCapture.cpp
+++ b/Engine/Plugins/Media/AjaMedia/Source/AjaMediaOutput/Private/AjaMediaCapture.cpp
@@ -283,7 +283,7 @@ void UAjaMediaCapture::StopCaptureImpl(bool bAllowPendingFrameToBeProcess)
 			// Prevent the rendering thread from copying while we are stopping the capture.
 			FScopeLock ScopeLock(&RenderThreadCriticalSection);
 
-			ENQUEUE_RENDER_COMMAND(BlackmagicMediaCaptureInitialize)(
+			ENQUEUE_RENDER_COMMAND(AjaMediaCaptureInitialize)(
 			[this](FRHICommandListImmediate& RHICmdList) mutable
 			{
 				// Unregister texture before closing channel.
diff --git a/Engine/Plugins/VirtualProduction/Rivermax/RivermaxMedia/RivermaxMedia.uplugin b/Engine/Plugins/VirtualProduction/Rivermax/RivermaxMedia/RivermaxMedia.uplugin
new file mode 100644
--- /dev/null
+++ b/Engine/Plugins/VirtualProduction/Rivermax/RivermaxMedia/RivermaxMedia.uplugin
@@ -0,0 +1,48 @@
+{
+	"FileVersion": 3,
+	"Version": 1,
+	"VersionName": "1.0",
+	"FriendlyName": "Rivermax Media",
+	"Description": "Adding Rivermax capabilities for Media Captures and Media Players",
+	"Category": "Virtual Production",
+	"CreatedBy": "",
+	"CreatedByURL": "",
+	"DocsURL": "",
+	"MarketplaceURL": "",
+	"SupportURL": "",
+	"Hidden": "false",
+	"CanContainContent": false,
+	"IsBetaVersion": true,
+	"Installed": false,
+	"SupportedTargetPlatforms": [ "Win64" ],
+	"Plugins": [
+		{
+			"Name": "MediaIOFramework",
+			"Enabled": true
+		},
+		{
+			"Name": "RivermaxCore",
+			"Enabled": true
+		}
+	],
+	"Modules": [
+		{
+			"Name": "RivermaxMedia",
+			"Type": "Runtime",
+			"LoadingPhase": "Default",
+			"PlatformAllowList": [ "Win64" ]
+		},
+		{
+			"Name": "RivermaxMediaEditor",
+			"Type": "Editor",
+			"LoadingPhase": "Default",
+			"PlatformAllowList": [ "Win64" ]
+		},
+		{
+			"Name": "RivermaxMediaFactory",
+			"Type": "RuntimeNoCommandlet",
+			"LoadingPhase": "PostEngineInit",
+			"PlatformAllowList": [ "Win64" ]
+		}
+	]
+}
\ No newline at end of file
diff --git a/Engine/Plugins/VirtualProduction/Rivermax/RivermaxMedia/Source/RivermaxMedia/Private/RivermaxMediaCapture.cpp b/Engine/Plugins/VirtualProduction/Rivermax/RivermaxMedia/Source/RivermaxMedia/Private/RivermaxMediaCapture.cpp
new file mode 100644
--- /dev/null
+++ b/Engine/Plugins/VirtualProduction/Rivermax/RivermaxMedia/Source/RivermaxMedia/Private/RivermaxMediaCapture.cpp
@@ -0,0 +1,213 @@
+// Copyright Epic Games, Inc. All Rights Reserved.
+
+#include "RivermaxMediaCapture.h"
+
+#include "IRivermaxCoreModule.h"
+#include "IRivermaxOutputStream.h"
+#include "RivermaxMediaLog.h"
+#include "RivermaxMediaOutput.h"
+#include "RivermaxTypes.h"
+#include "Slate/SceneViewport.h"
+#include "Widgets/SViewport.h"
+
+#if WITH_EDITOR
+#include "AnalyticsEventAttribute.h"
+#include "EngineAnalytics.h"
+#endif
+
+
+
+/* namespace RivermaxMediaCaptureDevice
+*****************************************************************************/
+
+#if WITH_EDITOR
+namespace RivermaxMediaCaptureAnalytics
+{
+	/**
+	 * @EventName MediaFramework.RivermaxCaptureStarted
+	 * @Trigger Triggered when a Rivermax capture of the viewport or render target is started.
+	 * @Type Client
+	 * @Owner MediaIO Team
+	 */
+	void SendCaptureEvent(const FIntPoint& Resolution, const FFrameRate FrameRate, const FString& CaptureType)
+	{
+		if (FEngineAnalytics::IsAvailable())
+		{
+			TArray<FAnalyticsEventAttribute> EventAttributes;
+			EventAttributes.Add(FAnalyticsEventAttribute(TEXT("CaptureType"), CaptureType));
+			EventAttributes.Add(FAnalyticsEventAttribute(TEXT("ResolutionWidth"), FString::Printf(TEXT("%d"), Resolution.X)));
+			EventAttributes.Add(FAnalyticsEventAttribute(TEXT("ResolutionHeight"), FString::Printf(TEXT("%d"), Resolution.Y)));
+			EventAttributes.Add(FAnalyticsEventAttribute(TEXT("FrameRate"), FrameRate.ToPrettyText().ToString()));
+			FEngineAnalytics::GetProvider().RecordEvent(TEXT("MediaFramework.RivermaxCaptureStarted"), EventAttributes);
+		}
+	}
+}
+#endif
+
+
+///* URivermaxMediaCapture implementation
+//*****************************************************************************/
+
+bool URivermaxMediaCapture::ValidateMediaOutput() const
+{
+	URivermaxMediaOutput* RivermaxMediaOutput = Cast<URivermaxMediaOutput>(MediaOutput);
+	if (!RivermaxMediaOutput)
+	{
+		UE_LOG(LogRivermaxMedia, Error, TEXT("Can not start the capture. MediaOutput's class is not supported."));
+		return false;
+	}
+
+	return true;
+}
+
+bool URivermaxMediaCapture::CaptureSceneViewportImpl(TSharedPtr<FSceneViewport>& InSceneViewport)
+{
+	URivermaxMediaOutput* RivermaxOutput = CastChecked<URivermaxMediaOutput>(MediaOutput);
+	const bool bResult = Initialize(RivermaxOutput);
+	if (bResult)
+	{
+#if WITH_EDITOR
+		RivermaxMediaCaptureAnalytics::SendCaptureEvent(RivermaxOutput->GetRequestedSize(), RivermaxOutput->FrameRate, TEXT("SceneViewport"));
+#endif
+	}
+	return bResult;
+}
+
+bool URivermaxMediaCapture::CaptureRenderTargetImpl(UTextureRenderTarget2D* InRenderTarget)
+{
+	URivermaxMediaOutput* RivermaxOutput = CastChecked<URivermaxMediaOutput>(MediaOutput);
+	bool bResult = Initialize(RivermaxOutput);
+#if WITH_EDITOR
+	if (bResult)
+	{
+		RivermaxMediaCaptureAnalytics::SendCaptureEvent(RivermaxOutput->GetRequestedSize(), RivermaxOutput->FrameRate, TEXT("RenderTarget"));
+	}
+#endif
+	return bResult;
+}
+
+bool URivermaxMediaCapture::UpdateSceneViewportImpl(TSharedPtr<FSceneViewport>& InSceneViewport)
+{
+	return true;
+}
+
+bool URivermaxMediaCapture::UpdateRenderTargetImpl(UTextureRenderTarget2D* InRenderTarget)
+{
+	return true;
+}
+
+void URivermaxMediaCapture::StopCaptureImpl(bool bAllowPendingFrameToBeProcess)
+{
+	if (RivermaxStream)
+	{
+		RivermaxStream->Uninitialize();
+		RivermaxStream.Reset();
+	}
+}
+
+bool URivermaxMediaCapture::ShouldCaptureRHITexture() const
+{
+	//todo for gpudirect support
+	return false;  
+}
+
+void URivermaxMediaCapture::BeforeFrameCaptured_RenderingThread(const FCaptureBaseData& InBaseData, TSharedPtr<FMediaCaptureUserData, ESPMode::ThreadSafe> InUserData, FTextureRHIRef InTexture)
+{
+	if (ShouldCaptureRHITexture())
+	{
+	
+	}
+}
+
+bool URivermaxMediaCapture::HasFinishedProcessing() const
+{
+	return Super::HasFinishedProcessing();
+}
+
+bool URivermaxMediaCapture::Initialize(URivermaxMediaOutput* InMediaOutput)
+{
+	using namespace UE::RivermaxCore;
+
+	check(InMediaOutput);
+
+	IRivermaxCoreModule* Module = FModuleManager::GetModulePtr<IRivermaxCoreModule>("RivermaxCore");
+	if (Module)
+	{
+		FRivermaxStreamOptions Options;
+		if (ConfigureStream(InMediaOutput, Options))
+		{
+			RivermaxStream = Module->CreateOutputStream();
+			if (RivermaxStream)
+			{
+				return RivermaxStream->Initialize(Options, *this);
+			}
+		}
+	}
+
+	return false;
+}
+
+bool URivermaxMediaCapture::ConfigureStream(URivermaxMediaOutput* InMediaOutput, UE::RivermaxCore::FRivermaxStreamOptions& OutOptions) const
+{
+	using namespace UE::RivermaxCore;
+
+	OutOptions.DestinationAddress = InMediaOutput->DestinationAddress;
+	OutOptions.SourceAddress = InMediaOutput->SourceAddress;
+	OutOptions.Port = InMediaOutput->Port;
+	OutOptions.Resolution = InMediaOutput->Resolution;
+	OutOptions.FrameRate = InMediaOutput->FrameRate;
+	OutOptions.NumberOfBuffers = InMediaOutput->NumberOfTextureBuffers;
+
+	switch (InMediaOutput->PixelFormat)
+	{
+	case ERivermaxMediaOutputPixelFormat::PF_8BIT_YUV:
+	{
+		OutOptions.PixelFormat = ERivermaxOutputPixelFormat::RMAX_8BIT_YCBCR;
+		break;
+	}
+	default:
+	{
+		UE_LOG(LogRivermaxMedia, Error, TEXT("Desired pixel format (%s) is not a valid Rivermax pixel format"), *StaticEnum<ERivermaxMediaOutputPixelFormat>()->GetValueAsString(InMediaOutput->PixelFormat));
+		return false;
+	}
+	}
+
+	return true;
+}
+
+void URivermaxMediaCapture::OnFrameCaptured_RenderingThread(const FCaptureBaseData& InBaseData, TSharedPtr<FMediaCaptureUserData, ESPMode::ThreadSafe> InUserData, void* InBuffer, int32 Width, int32 Height, int32 BytesPerRow)
+{
+	using namespace UE::RivermaxCore;
+
+	TRACE_CPUPROFILER_EVENT_SCOPE(URivermaxMediaCapture::OnFrameCaptured_RenderingThread);
+	
+	FRivermaxOutputVideoFrameInfo NewFrame;
+	NewFrame.Height = Height;
+	NewFrame.Width = Width;
+	NewFrame.Stride = BytesPerRow;
+	NewFrame.VideoBuffer = InBuffer;
+	NewFrame.FrameIdentifier = InBaseData.SourceFrameNumber;
+	RivermaxStream->PushVideoFrame(NewFrame);
+}
+
+void URivermaxMediaCapture::OnRHITextureCaptured_RenderingThread(const FCaptureBaseData& InBaseData,	TSharedPtr<FMediaCaptureUserData, ESPMode::ThreadSafe> InUserData, FTextureRHIRef InTexture)
+{
+	TRACE_CPUPROFILER_EVENT_SCOPE(URivermaxMediaCapture::OnFrameCaptured_RenderingThread);
+}
+
+void URivermaxMediaCapture::OnInitializationCompleted(bool bHasSucceed)
+{
+	if (GetState() != EMediaCaptureState::Stopped)
+	{
+		SetState(bHasSucceed ? EMediaCaptureState::Capturing : EMediaCaptureState::Error);
+	}
+}
+
+void URivermaxMediaCapture::OnStreamError()
+{
+	UE_LOG(LogRivermaxMedia, Error, TEXT("Outputstream has caught an error. Stopping capture."));
+	if (GetState() != EMediaCaptureState::Stopped)
+	{
+		SetState(EMediaCaptureState::Error);
+	}
+}
diff --git a/Engine/Plugins/VirtualProduction/Rivermax/RivermaxMedia/Source/RivermaxMedia/Private/RivermaxMediaLog.cpp b/Engine/Plugins/VirtualProduction/Rivermax/RivermaxMedia/Source/RivermaxMedia/Private/RivermaxMediaLog.cpp
new file mode 100644
--- /dev/null
+++ b/Engine/Plugins/VirtualProduction/Rivermax/RivermaxMedia/Source/RivermaxMedia/Private/RivermaxMediaLog.cpp
@@ -0,0 +1,5 @@
+// Copyright Epic Games, Inc. All Rights Reserved.
+
+#include "RivermaxMediaLog.h"
+
+DEFINE_LOG_CATEGORY(LogRivermaxMedia);
diff --git a/Engine/Plugins/VirtualProduction/Rivermax/RivermaxMedia/Source/RivermaxMedia/Private/RivermaxMediaLog.h b/Engine/Plugins/VirtualProduction/Rivermax/RivermaxMedia/Source/RivermaxMedia/Private/RivermaxMediaLog.h
new file mode 100644
--- /dev/null
+++ b/Engine/Plugins/VirtualProduction/Rivermax/RivermaxMedia/Source/RivermaxMedia/Private/RivermaxMediaLog.h
@@ -0,0 +1,7 @@
+// Copyright Epic Games, Inc. All Rights Reserved.
+
+#pragma once
+
+#include "CoreMinimal.h"
+
+DECLARE_LOG_CATEGORY_EXTERN(LogRivermaxMedia, Log, All);
diff --git a/Engine/Plugins/VirtualProduction/Rivermax/RivermaxMedia/Source/RivermaxMedia/Private/RivermaxMediaModule.cpp b/Engine/Plugins/VirtualProduction/Rivermax/RivermaxMedia/Source/RivermaxMedia/Private/RivermaxMediaModule.cpp
new file mode 100644
--- /dev/null
+++ b/Engine/Plugins/VirtualProduction/Rivermax/RivermaxMedia/Source/RivermaxMedia/Private/RivermaxMediaModule.cpp
@@ -0,0 +1,24 @@
+// Copyright Epic Games, Inc. All Rights Reserved.
+
+#include "RivermaxMediaModule.h"
+
+#include "Modules/ModuleManager.h"
+#include "RivermaxMediaPlayer.h"
+
+
+
+void FRivermaxMediaModule::StartupModule()
+{
+}
+
+void FRivermaxMediaModule::ShutdownModule()
+{
+
+}
+
+TSharedPtr<IMediaPlayer> FRivermaxMediaModule::CreatePlayer(IMediaEventSink& EventSink)
+{
+	return MakeShared<UE::RivermaxMedia::FRivermaxMediaPlayer>(EventSink);
+}
+
+IMPLEMENT_MODULE(FRivermaxMediaModule, RivermaxMedia);
diff --git a/Engine/Plugins/VirtualProduction/Rivermax/RivermaxMedia/Source/RivermaxMedia/Private/RivermaxMediaModule.h b/Engine/Plugins/VirtualProduction/Rivermax/RivermaxMedia/Source/RivermaxMedia/Private/RivermaxMediaModule.h
new file mode 100644
--- /dev/null
+++ b/Engine/Plugins/VirtualProduction/Rivermax/RivermaxMedia/Source/RivermaxMedia/Private/RivermaxMediaModule.h
@@ -0,0 +1,20 @@
+// Copyright Epic Games, Inc. All Rights Reserved.
+
+#pragma once
+
+#include "IRivermaxMediaModule.h"
+
+/**
+ * 
+ */
+class FRivermaxMediaModule : public IRivermaxMediaModule
+{
+public:
+	//~ Begin IModuleInterface
+	virtual void StartupModule() override;
+	virtual void ShutdownModule() override;
+	//~ End IModuleInterface
+
+	virtual TSharedPtr<IMediaPlayer> CreatePlayer(IMediaEventSink& EventSink) override;
+};
+
diff --git a/Engine/Plugins/VirtualProduction/Rivermax/RivermaxMedia/Source/RivermaxMedia/Private/RivermaxMediaOutput.cpp b/Engine/Plugins/VirtualProduction/Rivermax/RivermaxMedia/Source/RivermaxMedia/Private/RivermaxMediaOutput.cpp
new file mode 100644
--- /dev/null
+++ b/Engine/Plugins/VirtualProduction/Rivermax/RivermaxMedia/Source/RivermaxMedia/Private/RivermaxMediaOutput.cpp
@@ -0,0 +1,59 @@
+// Copyright Epic Games, Inc. All Rights Reserved.
+
+#include "RivermaxMediaOutput.h"
+
+#include "MediaOutput.h"
+#include "RivermaxMediaCapture.h"
+
+
+/* URivermaxMediaOutput
+*****************************************************************************/
+
+bool URivermaxMediaOutput::Validate(FString& OutFailureReason) const
+{
+	if (!Super::Validate(OutFailureReason))
+	{
+		return false;
+	}
+	
+	return true;
+}
+
+FIntPoint URivermaxMediaOutput::GetRequestedSize() const
+{
+	return Resolution;
+}
+
+EPixelFormat URivermaxMediaOutput::GetRequestedPixelFormat() const
+{
+	return EPixelFormat::PF_B8G8R8A8;
+}
+
+EMediaCaptureConversionOperation URivermaxMediaOutput::GetConversionOperation(EMediaCaptureSourceType InSourceType) const
+{
+	const EMediaCaptureConversionOperation Result = EMediaCaptureConversionOperation::RGBA8_TO_YUV_8BIT;
+	return Result;
+}
+
+UMediaCapture* URivermaxMediaOutput::CreateMediaCaptureImpl()
+{
+	UMediaCapture* Result = NewObject<URivermaxMediaCapture>();
+	if (Result)
+	{
+		Result->SetMediaOutput(this);
+	}
+	return Result;
+}
+
+#if WITH_EDITOR
+bool URivermaxMediaOutput::CanEditChange(const FProperty* InProperty) const
+{
+	return Super::CanEditChange(InProperty);
+}
+
+void URivermaxMediaOutput::PostEditChangeChainProperty(struct FPropertyChangedChainEvent& PropertyChangedEvent)
+{
+	Super::PostEditChangeChainProperty(PropertyChangedEvent);
+}
+#endif //WITH_EDITOR
+
diff --git a/Engine/Plugins/VirtualProduction/Rivermax/RivermaxMedia/Source/RivermaxMedia/Private/RivermaxMediaPlayer.cpp b/Engine/Plugins/VirtualProduction/Rivermax/RivermaxMedia/Source/RivermaxMedia/Private/RivermaxMediaPlayer.cpp
new file mode 100644
--- /dev/null
+++ b/Engine/Plugins/VirtualProduction/Rivermax/RivermaxMedia/Source/RivermaxMedia/Private/RivermaxMediaPlayer.cpp
@@ -0,0 +1,286 @@
+// Copyright Epic Games, Inc. All Rights Reserved.
+
+#include "RivermaxMediaPlayer.h"
+
+#include "IMediaEventSink.h"
+#include "IRivermaxCoreModule.h"
+#include "MediaIOCoreEncodeTime.h"
+#include "MediaIOCoreFileWriter.h"
+#include "MediaIOCoreSamples.h"
+#include "Misc/ScopeLock.h"
+#include "RivermaxMediaLog.h"
+#include "RivermaxMediaOutput.h"
+#include "RivermaxMediaSourceOptions.h"
+#include "RivermaxMediaTextureSample.h"
+#include "RivermaxTypes.h"
+#include "Stats/Stats2.h"
+
+#if WITH_EDITOR
+#include "EngineAnalytics.h"
+#endif
+
+#define LOCTEXT_NAMESPACE "FRivermaxMediaPlayer"
+
+DECLARE_CYCLE_STAT(TEXT("Rivermax MediaPlayer Request frame"), STAT_Rivermax_MediaPlayer_RequestFrame, STATGROUP_Media);
+DECLARE_CYCLE_STAT(TEXT("Rivermax MediaPlayer Process frame"), STAT_Rivermax_MediaPlayer_ProcessFrame, STATGROUP_Media);
+
+namespace UE::RivermaxMedia
+{
+	/* FRivermaxVideoPlayer structors
+	 *****************************************************************************/
+
+	FRivermaxMediaPlayer::FRivermaxMediaPlayer(IMediaEventSink& InEventSink)
+		: Super(InEventSink)
+		, TextureSamplePool(MakeUnique<FRivermaxMediaTextureSamplePool>())
+		, MaxNumVideoFrameBuffer(8)
+		, RivermaxThreadNewState(EMediaState::Closed)
+		, EventSink(InEventSink)
+		, bIsSRGBInput(true)
+		, bUseVideo(false)
+		, SupportedSampleTypes(EMediaIOSampleType::None)
+		, bPauseRequested(false)
+	{
+	}
+
+	FRivermaxMediaPlayer::~FRivermaxMediaPlayer()
+	{
+		Close();
+	}
+
+	/* IMediaPlayer interface
+	 *****************************************************************************/
+
+	 /**
+	  * @EventName MediaFramework.RivermaxSourceOpened
+	  * @Trigger Triggered when an Rivermax media source is opened through a media player.
+	  * @Type Client
+	  * @Owner MediaIO Team
+	  */
+	bool FRivermaxMediaPlayer::Open(const FString& Url, const IMediaOptions* Options)
+	{
+		if (!Super::Open(Url, Options))
+		{
+			return false;
+		}
+
+		{
+			StreamOptions.SourceAddress = Options->GetMediaOption(RivermaxMediaOption::SourceAddress, FString());
+			StreamOptions.DestinationAddress = Options->GetMediaOption(RivermaxMediaOption::DestinationAddress, FString());
+			StreamOptions.Port = Options->GetMediaOption(RivermaxMediaOption::Port, (int64)0);
+
+			ERivermaxMediaOutputPixelFormat DesiredFormat = ERivermaxMediaOutputPixelFormat::PF_8BIT_YUV;
+			DesiredFormat = (ERivermaxMediaOutputPixelFormat)Options->GetMediaOption(RivermaxMediaOption::PixelFormat, (int64)DesiredFormat);
+			StreamOptions.PixelFormat = DesiredFormat == ERivermaxMediaOutputPixelFormat::PF_8BIT_YUV ? ERivermaxOutputPixelFormat::RMAX_8BIT_YCBCR : ERivermaxOutputPixelFormat::RMAX_8BIT_YCBCR;
+			StreamOptions.Resolution = VideoTrackFormat.Dim;
+			StreamOptions.FrameRate = VideoFrameRate;
+
+		}
+		{
+			//Adjust supported sample types based on what's being captured
+			SupportedSampleTypes = EMediaIOSampleType::Video;
+			Samples->EnableTimedDataChannels(this, SupportedSampleTypes);
+		}
+
+		MaxNumVideoFrameBuffer = 8;
+		bUseVideo = true;
+
+		IRivermaxCoreModule* Module = FModuleManager::GetModulePtr<IRivermaxCoreModule>("RivermaxCore");
+		if (Module)
+		{
+			InputStream = Module->CreateInputStream();
+		}
+		if (InputStream == nullptr || !InputStream->Initialize(StreamOptions, *this))
+		{
+			UE_LOG(LogRivermaxMedia, Warning, TEXT("The Rivermax port couldn't be opened."));
+			CurrentState = EMediaState::Error;
+			RivermaxThreadNewState = EMediaState::Error;
+			InputStream.Reset();
+			return false;
+		}
+
+		// Setup our different supported channels based on source settings
+		SetupSampleChannels();
+
+		// finalize
+		CurrentState = EMediaState::Preparing;
+		RivermaxThreadNewState = EMediaState::Preparing;
+		EventSink.ReceiveMediaEvent(EMediaEvent::MediaConnecting);
+
+#if WITH_EDITOR
+		if (FEngineAnalytics::IsAvailable())
+		{
+			TArray<FAnalyticsEventAttribute> EventAttributes;
+
+			EventAttributes.Add(FAnalyticsEventAttribute(TEXT("ResolutionWidth"), FString::Printf(TEXT("%d"), VideoTrackFormat.Dim.X)));
+			EventAttributes.Add(FAnalyticsEventAttribute(TEXT("ResolutionHeight"), FString::Printf(TEXT("%d"), VideoTrackFormat.Dim.Y)));
+			EventAttributes.Add(FAnalyticsEventAttribute(TEXT("FrameRate"), *VideoFrameRate.ToPrettyText().ToString()));
+
+			FEngineAnalytics::GetProvider().RecordEvent(TEXT("MediaFramework.RivermaxSourceOpened"), EventAttributes);
+		}
+#endif
+
+		return true;
+	}
+
+	void FRivermaxMediaPlayer::Close()
+	{
+		RivermaxThreadNewState = EMediaState::Closed;
+
+		if (InputStream)
+		{
+			InputStream->Uninitialize(); // this may block, until the completion of a callback from IRivermaxChannelCallbackInterface
+			InputStream.Reset();
+		}
+
+		TextureSamplePool->Reset();
+
+		//Disable all our channels from the monitor
+		Samples->EnableTimedDataChannels(this, EMediaIOSampleType::None);
+
+		RivermaxThreadCurrentTextureSample.Reset();
+
+		Super::Close();
+	}
+
+	FGuid FRivermaxMediaPlayer::GetPlayerPluginGUID() const
+	{
+		static FGuid PlayerPluginGUID(0xF537595A, 0x8E8D452B, 0xB8C05707, 0x6B334234);
+		return PlayerPluginGUID;
+	}
+
+	FString FRivermaxMediaPlayer::GetStats() const
+	{
+		FString Stats;
+
+		//todo
+
+		return Stats;
+	}
+
+#if WITH_EDITOR
+	const FSlateBrush* FRivermaxMediaPlayer::GetDisplayIcon() const
+	{
+		//todo for tdm
+		return nullptr;
+		//return IRivermaxMediaModule::Get().GetStyle()->GetBrush("RivermaxMediaIcon");
+	}
+#endif //WITH_EDITOR
+
+	bool FRivermaxMediaPlayer::OnVideoFrameRequested(const FRivermaxInputVideoFrameDescriptor& FrameInfo, FRivermaxInputVideoFrameRequest& OutVideoFrameRequest)
+	{
+		if (FrameInfo.VideoBufferSize > 0)
+		{
+			RivermaxThreadCurrentTextureSample = TextureSamplePool->AcquireShared();
+			OutVideoFrameRequest.VideoBuffer = reinterpret_cast<uint8*>(RivermaxThreadCurrentTextureSample->RequestBuffer(FrameInfo.VideoBufferSize));
+			return OutVideoFrameRequest.VideoBuffer != nullptr;
+		}
+
+		return false;
+	}
+
+
+	void FRivermaxMediaPlayer::OnVideoFrameReceived(const FRivermaxInputVideoFrameDescriptor& FrameInfo, const FRivermaxInputVideoFrameReception& ReceivedVideoFrame)
+	{
+		FTimespan DecodedTime = FTimespan::FromSeconds(GetPlatformSeconds());
+		TOptional<FTimecode> DecodedTimecode;
+
+
+		if (bUseVideo && ReceivedVideoFrame.VideoBuffer)
+		{
+			const EMediaTextureSampleFormat VideoSampleFormat = EMediaTextureSampleFormat::CharUYVY;
+			if (RivermaxThreadCurrentTextureSample.IsValid())
+			{
+				if (RivermaxThreadCurrentTextureSample->SetProperties(FrameInfo.Stride, FrameInfo.Width, FrameInfo.Height, VideoSampleFormat, DecodedTime, VideoFrameRate, DecodedTimecode, bIsSRGBInput))
+				{
+					Samples->AddVideo(RivermaxThreadCurrentTextureSample.ToSharedRef());
+				}
+			}
+		}
+
+		RivermaxThreadCurrentTextureSample.Reset();
+	}
+
+	void FRivermaxMediaPlayer::TickFetch(FTimespan DeltaTime, FTimespan Timecode)
+	{
+		Super::TickFetch(DeltaTime, Timecode);
+		if (InputStream && CurrentState == EMediaState::Playing)
+		{
+			ProcessFrame();
+		}
+	}
+
+	void FRivermaxMediaPlayer::TickInput(FTimespan DeltaTime, FTimespan Timecode)
+	{
+		// update player state
+		EMediaState NewState = RivermaxThreadNewState;
+
+		if (NewState != CurrentState)
+		{
+			CurrentState = NewState;
+			if (CurrentState == EMediaState::Playing)
+			{
+				EventSink.ReceiveMediaEvent(EMediaEvent::TracksChanged);
+				EventSink.ReceiveMediaEvent(EMediaEvent::MediaOpened);
+				EventSink.ReceiveMediaEvent(EMediaEvent::PlaybackResumed);
+			}
+			else if (NewState == EMediaState::Error)
+			{
+				EventSink.ReceiveMediaEvent(EMediaEvent::MediaOpenFailed);
+				Close();
+			}
+		}
+
+		if (CurrentState != EMediaState::Playing)
+		{
+			return;
+		}
+
+		TickTimeManagement();
+	}
+
+
+	/* FRivermaxMediaPlayer implementation
+	 *****************************************************************************/
+	void FRivermaxMediaPlayer::ProcessFrame()
+	{
+		if (CurrentState == EMediaState::Playing)
+		{
+		}
+	}
+
+	bool FRivermaxMediaPlayer::IsHardwareReady() const
+	{
+		return (RivermaxThreadNewState == EMediaState::Playing) || (RivermaxThreadNewState == EMediaState::Paused);
+	}
+
+	void FRivermaxMediaPlayer::SetupSampleChannels()
+	{
+		FMediaIOSamplingSettings VideoSettings = BaseSettings;
+		VideoSettings.BufferSize = MaxNumVideoFrameBuffer;
+		Samples->InitializeVideoBuffer(VideoSettings);
+	}
+
+	bool FRivermaxMediaPlayer::SetRate(float Rate)
+	{
+		if (FMath::IsNearlyEqual(Rate, 1.0f))
+		{
+			bPauseRequested = false;
+			return true;
+		}
+
+		if (FMath::IsNearlyEqual(Rate, 0.0f))
+		{
+			bPauseRequested = true;
+			return true;
+		}
+
+		return false;
+	}
+
+	void FRivermaxMediaPlayer::OnInitializationCompleted(bool bHasSucceed)
+	{
+		RivermaxThreadNewState = bHasSucceed ? EMediaState::Playing : EMediaState::Error;
+	}
+}
+
+#undef LOCTEXT_NAMESPACE
diff --git a/Engine/Plugins/VirtualProduction/Rivermax/RivermaxMedia/Source/RivermaxMedia/Private/RivermaxMediaPlayer.h b/Engine/Plugins/VirtualProduction/Rivermax/RivermaxMedia/Source/RivermaxMedia/Private/RivermaxMediaPlayer.h
new file mode 100644
--- /dev/null
+++ b/Engine/Plugins/VirtualProduction/Rivermax/RivermaxMedia/Source/RivermaxMedia/Private/RivermaxMediaPlayer.h
@@ -0,0 +1,120 @@
+// Copyright Epic Games, Inc. All Rights Reserved.
+
+#pragma once
+
+#include "MediaIOCorePlayerBase.h"
+#include "IRivermaxInputStream.h"
+#include "RivermaxMediaSource.h"
+#include "RivermaxTypes.h"
+
+#include <atomic>
+
+class FRivermaxMediaTextureSample;
+class FRivermaxMediaTextureSamplePool;
+struct FSlateBrush;
+class IMediaEventSink;
+
+enum class EMediaTextureSampleFormat;
+enum class EMediaIOSampleType;
+
+
+namespace  UE::RivermaxMedia
+{
+	using namespace UE::RivermaxCore;
+
+	/**
+	 * Implements a media player using rivermax.
+	 */
+	class FRivermaxMediaPlayer : public FMediaIOCorePlayerBase, public IRivermaxInputStreamListener
+	{
+		using Super = FMediaIOCorePlayerBase;
+	public:
+
+		/**
+		 * Create and initialize a new instance.
+		 *
+		 * @param InEventSink The object that receives media events from this player.
+		 */
+		FRivermaxMediaPlayer(IMediaEventSink& InEventSink);
+
+		/** Virtual destructor. */
+		virtual ~FRivermaxMediaPlayer();
+
+	public:
+
+		//~ IMediaPlayer interface
+
+		virtual void Close() override;
+		virtual FGuid GetPlayerPluginGUID() const override;
+
+		virtual bool Open(const FString& Url, const IMediaOptions* Options) override;
+
+		virtual void TickFetch(FTimespan DeltaTime, FTimespan Timecode) override;
+		virtual void TickInput(FTimespan DeltaTime, FTimespan Timecode) override;
+
+		virtual FString GetStats() const override;
+
+		bool SetRate(float Rate) override;
+
+		//~ ITimedDataInput interface
+#if WITH_EDITOR
+		virtual const FSlateBrush* GetDisplayIcon() const override;
+#endif
+
+		//~ Begin IRivermaxInputStreamListener interface
+		virtual void OnInitializationCompleted(bool bHasSucceed) override;
+		virtual bool OnVideoFrameRequested(const FRivermaxInputVideoFrameDescriptor& FrameInfo, FRivermaxInputVideoFrameRequest& OutVideoFrameRequest) override;
+		virtual void OnVideoFrameReceived(const FRivermaxInputVideoFrameDescriptor& FrameInfo, const FRivermaxInputVideoFrameReception& ReceivedVideoFrame) override;
+		//~ End IRivermaxInputStreamListener interface
+
+	protected:
+
+		/**
+		 * Process pending audio and video frames, and forward them to the sinks.
+		 */
+		void ProcessFrame();
+
+	protected:
+
+		//~ Begin FMediaIOCorePlayerBase interface
+		virtual bool IsHardwareReady() const override;
+		virtual void SetupSampleChannels() override;
+		//~ End FMediaIOCorePlayerBase interface
+
+	private:
+
+		/** Audio, MetaData, Texture  sample object pool. */
+		TUniquePtr<FRivermaxMediaTextureSamplePool> TextureSamplePool;
+
+		TSharedPtr<FRivermaxMediaTextureSample> RivermaxThreadCurrentTextureSample;
+
+		/** The media sample cache. */
+		int32 MaxNumVideoFrameBuffer;
+
+		/** Current state of the media player. */
+		EMediaState RivermaxThreadNewState;
+
+		/** The media event handler. */
+		IMediaEventSink& EventSink;
+
+		FRivermaxStreamOptions StreamOptions;
+
+		/** Whether the input is in sRGB and can have a ToLinear conversion. */
+		bool bIsSRGBInput;
+
+		/** Which field need to be capture. */
+		bool bUseVideo;
+		bool bVerifyFrameDropCount;
+
+		/** Maps to the current input Device */
+		TUniquePtr<IRivermaxInputStream> InputStream;
+
+		/** Used to flag which sample types we advertise as supported for timed data monitoring */
+		EMediaIOSampleType SupportedSampleTypes;
+
+		/** Flag to indicate that pause is being requested */
+		std::atomic<bool> bPauseRequested;
+	};
+}
+
+
diff --git a/Engine/Plugins/VirtualProduction/Rivermax/RivermaxMedia/Source/RivermaxMedia/Private/RivermaxMediaSource.cpp b/Engine/Plugins/VirtualProduction/Rivermax/RivermaxMedia/Source/RivermaxMedia/Private/RivermaxMediaSource.cpp
new file mode 100644
--- /dev/null
+++ b/Engine/Plugins/VirtualProduction/Rivermax/RivermaxMedia/Source/RivermaxMedia/Private/RivermaxMediaSource.cpp
@@ -0,0 +1,95 @@
+// Copyright Epic Games, Inc. All Rights Reserved.
+
+#include "RivermaxMediaSource.h"
+
+#include "MediaIOCorePlayerBase.h"
+#include "RivermaxMediaSourceOptions.h"
+
+
+/*
+ * IMediaOptions interface
+ */
+
+bool URivermaxMediaSource::GetMediaOption(const FName& Key, bool DefaultValue) const
+{
+	return Super::GetMediaOption(Key, DefaultValue);
+}
+
+int64 URivermaxMediaSource::GetMediaOption(const FName& Key, int64 DefaultValue) const
+{
+	if (Key == RivermaxMediaOption::Port)
+	{
+		return Port;
+	}
+	else if (Key == RivermaxMediaOption::PixelFormat)
+	{
+		return (int64)PixelFormat;
+	}
+	else if (Key == FMediaIOCoreMediaOption::FrameRateNumerator)
+	{
+		return FrameRate.Numerator;
+	}
+	else if (Key == FMediaIOCoreMediaOption::FrameRateDenominator)
+	{
+		return FrameRate.Denominator;
+	}
+	else if (Key == FMediaIOCoreMediaOption::ResolutionWidth)
+	{
+		return Resolution.X;
+	}
+	else if (Key == FMediaIOCoreMediaOption::ResolutionHeight)
+	{
+		return Resolution.Y;
+	}
+	return Super::GetMediaOption(Key, DefaultValue);
+}
+
+FString URivermaxMediaSource::GetMediaOption(const FName& Key, const FString& DefaultValue) const
+{
+	if (Key == FMediaIOCoreMediaOption::VideoModeName)
+	{
+		return FString::Printf(TEXT("FormatDescriptorTodo"));
+	}
+	else if (Key == RivermaxMediaOption::SourceAddress)
+	{
+		return SourceAddress;
+	}
+	else if (Key == RivermaxMediaOption::DestinationAddress)
+	{
+		return DestinationAddress;
+	}
+	return Super::GetMediaOption(Key, DefaultValue);
+}
+
+bool URivermaxMediaSource::HasMediaOption(const FName& Key) const
+{
+	if ((Key == RivermaxMediaOption::SourceAddress) ||
+		(Key == RivermaxMediaOption::DestinationAddress) ||
+		(Key == RivermaxMediaOption::Port) ||
+		(Key == RivermaxMediaOption::PixelFormat) ||
+		(Key == FMediaIOCoreMediaOption::FrameRateNumerator) ||
+		(Key == FMediaIOCoreMediaOption::FrameRateDenominator) ||
+		(Key == FMediaIOCoreMediaOption::ResolutionWidth) ||
+		(Key == FMediaIOCoreMediaOption::ResolutionHeight) ||
+		(Key == FMediaIOCoreMediaOption::VideoModeName))
+	{
+		return true;
+	}
+
+	return Super::HasMediaOption(Key);
+}
+
+/*
+ * UMediaSource interface
+ */
+
+FString URivermaxMediaSource::GetUrl() const
+{
+	return TEXT("rmax://");//todo support proper url
+}
+
+bool URivermaxMediaSource::Validate() const
+{
+	return true;
+}
+
diff --git a/Engine/Plugins/VirtualProduction/Rivermax/RivermaxMedia/Source/RivermaxMedia/Private/RivermaxMediaSourceOptions.h b/Engine/Plugins/VirtualProduction/Rivermax/RivermaxMedia/Source/RivermaxMedia/Private/RivermaxMediaSourceOptions.h
new file mode 100644
--- /dev/null
+++ b/Engine/Plugins/VirtualProduction/Rivermax/RivermaxMedia/Source/RivermaxMedia/Private/RivermaxMediaSourceOptions.h
@@ -0,0 +1,14 @@
+// Copyright Epic Games, Inc. All Rights Reserved.
+
+#pragma once
+
+#include "CoreMinimal.h"
+
+
+namespace RivermaxMediaOption
+{
+	static const FName SourceAddress("SourceAddress");
+	static const FName DestinationAddress("DestinationAddress");
+	static const FName Port("Port");
+	static const FName PixelFormat("PixelFormat");
+}
diff --git a/Engine/Plugins/VirtualProduction/Rivermax/RivermaxMedia/Source/RivermaxMedia/Private/RivermaxMediaTextureSample.h b/Engine/Plugins/VirtualProduction/Rivermax/RivermaxMedia/Source/RivermaxMedia/Private/RivermaxMediaTextureSample.h
new file mode 100644
--- /dev/null
+++ b/Engine/Plugins/VirtualProduction/Rivermax/RivermaxMedia/Source/RivermaxMedia/Private/RivermaxMediaTextureSample.h
@@ -0,0 +1,57 @@
+// Copyright Epic Games, Inc. All Rights Reserved.
+
+#pragma once
+
+#include "MediaIOCoreTextureSampleBase.h"
+#include "MediaShaders.h"
+
+/**
+ * Implements a media texture sample for RivermaxMedia.
+ */
+class FRivermaxMediaTextureSample
+	: public FMediaIOCoreTextureSampleBase
+{
+	using Super = FMediaIOCoreTextureSampleBase;
+
+public:
+
+	/**
+	 * Initialize the sample.
+	 *
+	 * @param InVideoData The video frame data.
+	 * @param InSampleFormat The sample format.
+	 * @param InTime The sample time (in the player's own clock).
+	 * @param InFrameRate The framerate of the media that produce the sample.
+	 * @param InTimecode The sample timecode if available.
+	 * @param bInIsSRGB Whether the sample is in sRGB space.
+	 */
+	bool InitializeProgressive(void* InVideoBuffer, uint32 InWidth, uint32 InHeight, uint32 InVideoBufferSize, uint32 InStride, EMediaTextureSampleFormat InSampleFormat, FTimespan InTime, const FFrameRate& InFrameRate, const TOptional<FTimecode>& InTimecode, bool bInIsSRGB)
+	{
+		return Super::Initialize(InVideoBuffer
+			, InVideoBufferSize
+			, InStride
+			, InWidth
+			, InHeight
+			, InSampleFormat
+			, InTime
+			, InFrameRate
+			, InTimecode
+			, bInIsSRGB);
+	}
+
+	/**
+	 * Get YUV to RGB conversion matrix
+	 *
+	 * @return MediaIOCore Yuv To Rgb matrix
+	 */
+	virtual const FMatrix& GetYUVToRGBMatrix() const override
+	{
+		return MediaShaders::YuvToRgbRec709Scaled;
+	}
+
+};
+
+/*
+ * Implements a pool for Rivermax texture sample objects.
+ */
+class FRivermaxMediaTextureSamplePool : public TMediaObjectPool<FRivermaxMediaTextureSample> { };
diff --git a/Engine/Plugins/VirtualProduction/Rivermax/RivermaxMedia/Source/RivermaxMedia/Public/IRivermaxMediaModule.h b/Engine/Plugins/VirtualProduction/Rivermax/RivermaxMedia/Source/RivermaxMedia/Public/IRivermaxMediaModule.h
new file mode 100644
--- /dev/null
+++ b/Engine/Plugins/VirtualProduction/Rivermax/RivermaxMedia/Source/RivermaxMedia/Public/IRivermaxMediaModule.h
@@ -0,0 +1,37 @@
+// Copyright Epic Games, Inc. All Rights Reserved.
+
+#pragma once
+
+#include "Modules/ModuleInterface.h"
+
+#include "Modules/ModuleManager.h"
+
+class IMediaPlayer;
+class IMediaEventSink;
+
+
+
+/**
+ * 
+ */
+class IRivermaxMediaModule : public IModuleInterface
+{
+public:
+	static inline IRivermaxMediaModule& Get()
+	{
+		return FModuleManager::LoadModuleChecked<IRivermaxMediaModule>("RivermaxMedia");
+	}
+
+	static inline bool IsAvailable()
+	{
+		return FModuleManager::Get().IsModuleLoaded("RivermaxMedia");
+	}
+
+	/**
+	 * Create a Rivermax based media player.
+	 * @param EventSink The object that receives media events from the player.
+	 * @return A new media player, or nullptr if a player couldn't be created.
+	 */
+	virtual TSharedPtr<IMediaPlayer> CreatePlayer(IMediaEventSink& EventSink) = 0;
+};
+
diff --git a/Engine/Plugins/VirtualProduction/Rivermax/RivermaxMedia/Source/RivermaxMedia/Public/RivermaxMediaCapture.h b/Engine/Plugins/VirtualProduction/Rivermax/RivermaxMedia/Source/RivermaxMedia/Public/RivermaxMediaCapture.h
new file mode 100644
--- /dev/null
+++ b/Engine/Plugins/VirtualProduction/Rivermax/RivermaxMedia/Source/RivermaxMedia/Public/RivermaxMediaCapture.h
@@ -0,0 +1,59 @@
+// Copyright Epic Games, Inc. All Rights Reserved.
+
+#pragma once
+
+#include "MediaCapture.h"
+
+#include "IRivermaxOutputStream.h"
+#include "Misc/FrameRate.h"
+#include "RivermaxMediaOutput.h"
+
+#include "RivermaxMediaCapture.generated.h"
+
+
+namespace UE::RivermaxCore
+{
+	struct FRivermaxStreamOptions;
+}
+
+class URivermaxMediaOutput;
+
+/**
+ * Output Media for Rivermax streams.
+ */
+UCLASS(BlueprintType)
+class RIVERMAXMEDIA_API URivermaxMediaCapture : public UMediaCapture, public UE::RivermaxCore::IRivermaxOutputStreamListener
+{
+	GENERATED_BODY()
+
+public:
+
+	//~ Begin UMediaCapture interface
+	virtual bool HasFinishedProcessing() const override;
+protected:
+	virtual bool ValidateMediaOutput() const override;
+	virtual bool CaptureSceneViewportImpl(TSharedPtr<FSceneViewport>& InSceneViewport) override;
+	virtual bool CaptureRenderTargetImpl(UTextureRenderTarget2D* InRenderTarget) override;
+	virtual bool UpdateSceneViewportImpl(TSharedPtr<FSceneViewport>& InSceneViewport) override;
+	virtual bool UpdateRenderTargetImpl(UTextureRenderTarget2D* InRenderTarget) override;
+	virtual void StopCaptureImpl(bool bAllowPendingFrameToBeProcess) override;
+	virtual bool ShouldCaptureRHITexture() const override;
+	
+	virtual void BeforeFrameCaptured_RenderingThread(const FCaptureBaseData& InBaseData, TSharedPtr<FMediaCaptureUserData, ESPMode::ThreadSafe> InUserData, FTextureRHIRef InTexture) override;
+	virtual void OnFrameCaptured_RenderingThread(const FCaptureBaseData& InBaseData, TSharedPtr<FMediaCaptureUserData, ESPMode::ThreadSafe> InUserData, void* InBuffer, int32 Width, int32 Height, int32 BytesPerRow) override;
+	virtual void OnRHITextureCaptured_RenderingThread(const FCaptureBaseData& InBaseData, TSharedPtr<FMediaCaptureUserData, ESPMode::ThreadSafe> InUserData, FTextureRHIRef InTexture) override;
+	//~ End UMediaCapture interface
+
+	//~ Begin IRivermaxOutputStreamListener interface
+	virtual void OnInitializationCompleted(bool bHasSucceed) override;
+	virtual void OnStreamError() override;
+	//~ End IRivermaxOutputStreamListener interface
+
+private:
+	bool Initialize(URivermaxMediaOutput* InMediaOutput);
+	bool ConfigureStream(URivermaxMediaOutput* InMediaOutput, UE::RivermaxCore::FRivermaxStreamOptions& OutOptions) const;
+
+private:
+
+	TUniquePtr<UE::RivermaxCore::IRivermaxOutputStream> RivermaxStream;
+};
diff --git a/Engine/Plugins/VirtualProduction/Rivermax/RivermaxMedia/Source/RivermaxMedia/Public/RivermaxMediaOutput.h b/Engine/Plugins/VirtualProduction/Rivermax/RivermaxMedia/Source/RivermaxMedia/Public/RivermaxMediaOutput.h
new file mode 100644
--- /dev/null
+++ b/Engine/Plugins/VirtualProduction/Rivermax/RivermaxMedia/Source/RivermaxMedia/Public/RivermaxMediaOutput.h
@@ -0,0 +1,67 @@
+// Copyright Epic Games, Inc. All Rights Reserved.
+
+#pragma once
+
+#include "MediaOutput.h"
+
+#include "MediaIOCoreDefinitions.h"
+#include "UObject/ObjectMacros.h"
+
+#include "RivermaxMediaOutput.generated.h"
+
+
+UENUM()
+enum class ERivermaxMediaOutputPixelFormat : uint8
+{
+	PF_8BIT_YUV UMETA(DisplayName = "8bit YUV"),
+	//Todo add 10bit YUV, 8/10 bit RGB
+};
+
+
+/**
+ * Output information for a Rivermax media capture.
+ */
+UCLASS(BlueprintType, meta=(MediaIOCustomLayout="Rivermax"))
+class RIVERMAXMEDIA_API URivermaxMediaOutput : public UMediaOutput
+{
+	GENERATED_BODY()
+
+public:
+	
+	//~ Begin UMediaOutput interface
+	virtual bool Validate(FString& FailureReason) const override;
+	virtual FIntPoint GetRequestedSize() const override;
+	virtual EPixelFormat GetRequestedPixelFormat() const override;
+	virtual EMediaCaptureConversionOperation GetConversionOperation(EMediaCaptureSourceType InSourceType) const override;
+
+protected:
+	virtual UMediaCapture* CreateMediaCaptureImpl() override;
+	//~ End UMediaOutput interface
+
+
+public:
+	//~ UObject interface
+#if WITH_EDITOR
+	virtual bool CanEditChange(const FProperty* InProperty) const override;
+	virtual void PostEditChangeChainProperty(struct FPropertyChangedChainEvent& PropertyChangedEvent) override;
+#endif //WITH_EDITOR
+	//~ End UObject interface
+
+	UPROPERTY(EditAnywhere, BlueprintReadWrite, Category="Format")
+	FIntPoint Resolution = {1920, 1080};
+	
+	UPROPERTY(EditAnywhere, BlueprintReadWrite, Category = "Format")
+	FFrameRate FrameRate = {24,1};
+	
+	UPROPERTY(EditAnywhere, BlueprintReadWrite, Category = "Format")
+	ERivermaxMediaOutputPixelFormat PixelFormat = ERivermaxMediaOutputPixelFormat::PF_8BIT_YUV;
+
+	UPROPERTY(EditAnywhere, BlueprintReadWrite, Category = "Format")
+	FString SourceAddress;
+
+	UPROPERTY(EditAnywhere, BlueprintReadWrite, Category = "Format")
+	FString DestinationAddress;
+
+	UPROPERTY(EditAnywhere, BlueprintReadWrite, Category = "Format")
+	int32 Port = 50000;
+};
diff --git a/Engine/Plugins/VirtualProduction/Rivermax/RivermaxMedia/Source/RivermaxMedia/Public/RivermaxMediaSource.h b/Engine/Plugins/VirtualProduction/Rivermax/RivermaxMedia/Source/RivermaxMedia/Public/RivermaxMediaSource.h
new file mode 100644
--- /dev/null
+++ b/Engine/Plugins/VirtualProduction/Rivermax/RivermaxMedia/Source/RivermaxMedia/Public/RivermaxMediaSource.h
@@ -0,0 +1,64 @@
+// Copyright Epic Games, Inc. All Rights Reserved.
+
+#pragma once
+
+#include "TimeSynchronizableMediaSource.h"
+
+#include "MediaIOCoreDefinitions.h"
+
+#include "RivermaxMediaSource.generated.h"
+
+/**
+ * Native data format.
+ */
+UENUM()
+enum class ERivermaxMediaSourceColorFormat : uint8
+{
+	YUV2_8bit UMETA(DisplayName = "8bit YUV"),
+	//Todo Add support for 10bit YUV and 8/10 bit RGB
+};
+
+/**
+ * Media source for Rivermax streams.
+ */
+UCLASS(BlueprintType, hideCategories=(Platforms,Object), meta=(MediaIOCustomLayout="Rivermax"))
+class RIVERMAXMEDIA_API URivermaxMediaSource : public UTimeSynchronizableMediaSource
+{
+	GENERATED_BODY()
+
+public:
+
+	//todo proper device configuration like Aja and Blackmagic
+
+	UPROPERTY(EditAnywhere, BlueprintReadWrite, Category="Format")
+	FIntPoint Resolution = {1920, 1080};
+	
+	UPROPERTY(EditAnywhere, BlueprintReadWrite, Category = "Format")
+	FFrameRate FrameRate = {24,1};
+	
+	UPROPERTY(EditAnywhere, BlueprintReadWrite, Category = "Format")
+	ERivermaxMediaSourceColorFormat PixelFormat = ERivermaxMediaSourceColorFormat::YUV2_8bit;
+
+	UPROPERTY(EditAnywhere, BlueprintReadWrite, Category = "Format")
+	FString SourceAddress;
+
+	UPROPERTY(EditAnywhere, BlueprintReadWrite, Category = "Format")
+	FString DestinationAddress;
+
+	UPROPERTY(EditAnywhere, BlueprintReadWrite, Category = "Format")
+	int32 Port = 50000;
+
+public:
+	//~ Begin IMediaOptions interface
+	virtual bool GetMediaOption(const FName& Key, bool DefaultValue) const override;
+	virtual int64 GetMediaOption(const FName& Key, int64 DefaultValue) const override;
+	virtual FString GetMediaOption(const FName& Key, const FString& DefaultValue) const override;
+	virtual bool HasMediaOption(const FName& Key) const override;
+	//~ End IMediaOptions interface
+
+public:
+	//~ Begin UMediaSource interface
+	virtual FString GetUrl() const override;
+	virtual bool Validate() const override;
+	//~ End UMediaSource interface
+};
diff --git a/Engine/Plugins/VirtualProduction/Rivermax/RivermaxMedia/Source/RivermaxMedia/RivermaxMedia.build.cs b/Engine/Plugins/VirtualProduction/Rivermax/RivermaxMedia/Source/RivermaxMedia/RivermaxMedia.build.cs
new file mode 100644
--- /dev/null
+++ b/Engine/Plugins/VirtualProduction/Rivermax/RivermaxMedia/Source/RivermaxMedia/RivermaxMedia.build.cs
@@ -0,0 +1,41 @@
+// Copyright Epic Games, Inc. All Rights Reserved.
+
+using UnrealBuildTool;
+
+public class RivermaxMedia : ModuleRules
+{
+	public RivermaxMedia(ReadOnlyTargetRules Target) : base(Target)
+	{
+		PublicDependencyModuleNames.AddRange(
+			new string[] {
+				"Core",
+				"CoreUObject",
+				"MediaAssets",
+				"MediaIOCore",
+				"TimeManagement"
+			}
+		);
+
+		PrivateDependencyModuleNames.AddRange(
+			new string[] {
+				"Engine",
+				"RenderCore",
+				"RHI",
+				"RivermaxCore"
+			}
+		);
+
+		if (Target.bBuildEditor)
+		{
+			PublicDependencyModuleNames.AddRange(
+				new string[] {
+				}
+			);
+
+			PrivateDependencyModuleNames.AddRange(
+				new string[] {
+				}
+			);
+		}
+	}
+}
diff --git a/Engine/Plugins/VirtualProduction/Rivermax/RivermaxMedia/Source/RivermaxMediaEditor/Private/RivermaxMediaEditorModule.cpp b/Engine/Plugins/VirtualProduction/Rivermax/RivermaxMedia/Source/RivermaxMediaEditor/Private/RivermaxMediaEditorModule.cpp
new file mode 100644
--- /dev/null
+++ b/Engine/Plugins/VirtualProduction/Rivermax/RivermaxMedia/Source/RivermaxMediaEditor/Private/RivermaxMediaEditorModule.cpp
@@ -0,0 +1,5 @@
+// Copyright Epic Games, Inc. All Rights Reserved.
+
+#include "Modules/ModuleManager.h"
+
+IMPLEMENT_MODULE(FDefaultModuleImpl, RivermaxMediaEditor);
diff --git a/Engine/Plugins/VirtualProduction/Rivermax/RivermaxMedia/Source/RivermaxMediaEditor/Private/RivermaxMediaOutputFactoryNew.cpp b/Engine/Plugins/VirtualProduction/Rivermax/RivermaxMedia/Source/RivermaxMediaEditor/Private/RivermaxMediaOutputFactoryNew.cpp
new file mode 100644
--- /dev/null
+++ b/Engine/Plugins/VirtualProduction/Rivermax/RivermaxMedia/Source/RivermaxMediaEditor/Private/RivermaxMediaOutputFactoryNew.cpp
@@ -0,0 +1,50 @@
+// Copyright Epic Games, Inc. All Rights Reserved.
+
+#include "RivermaxMediaOutputFactoryNew.h"
+
+#include "AssetTypeCategories.h"
+#include "RivermaxMediaOutput.h"
+#include "EngineAnalytics.h"
+
+
+/* UAjaMediaSourceFactoryNew structors
+ *****************************************************************************/
+
+URivermaxMediaOutputFactoryNew::URivermaxMediaOutputFactoryNew()
+{
+	SupportedClass = URivermaxMediaOutput::StaticClass();
+	bCreateNew = true;
+	bEditAfterNew = true;
+}
+
+
+/* UFactory overrides
+ *****************************************************************************/
+
+/**
+ * @EventName MediaFramework.RivermaxMediaOutputCreated
+ * @Trigger Triggered when an Rivermax media output asset is created.
+ * @Type Client
+ * @Owner MediaIO Team
+ */
+UObject* URivermaxMediaOutputFactoryNew::FactoryCreateNew(UClass* InClass, UObject* InParent, FName InName, EObjectFlags Flags, UObject* Context, FFeedbackContext* Warn)
+{
+	if (FEngineAnalytics::IsAvailable())
+	{
+		FEngineAnalytics::GetProvider().RecordEvent(TEXT("MediaFramework.RivermaxMediaOutputCreated"));
+	}
+	
+	return NewObject<URivermaxMediaOutput>(InParent, InClass, InName, Flags);
+}
+
+
+uint32 URivermaxMediaOutputFactoryNew::GetMenuCategories() const
+{
+	return EAssetTypeCategories::Media;
+}
+
+
+bool URivermaxMediaOutputFactoryNew::ShouldShowInNewMenu() const
+{
+	return true;
+}
diff --git a/Engine/Plugins/VirtualProduction/Rivermax/RivermaxMedia/Source/RivermaxMediaEditor/Private/RivermaxMediaOutputFactoryNew.h b/Engine/Plugins/VirtualProduction/Rivermax/RivermaxMedia/Source/RivermaxMediaEditor/Private/RivermaxMediaOutputFactoryNew.h
new file mode 100644
--- /dev/null
+++ b/Engine/Plugins/VirtualProduction/Rivermax/RivermaxMedia/Source/RivermaxMediaEditor/Private/RivermaxMediaOutputFactoryNew.h
@@ -0,0 +1,27 @@
+// Copyright Epic Games, Inc. All Rights Reserved.
+
+#pragma once
+
+#include "Factories/Factory.h"
+#include "UObject/ObjectMacros.h"
+
+#include "RivermaxMediaOutputFactoryNew.generated.h"
+
+
+/**
+ * Implements a factory for URivermaxMediaOutput objects.
+ */
+UCLASS(hidecategories=Object)
+class URivermaxMediaOutputFactoryNew : public UFactory
+{
+	GENERATED_BODY()
+
+public:
+
+	URivermaxMediaOutputFactoryNew();
+
+	//~ UFactory Interface
+	virtual UObject* FactoryCreateNew(UClass* InClass, UObject* InParent, FName InName, EObjectFlags Flags, UObject* Context, FFeedbackContext* Warn) override;
+	virtual uint32 GetMenuCategories() const override;
+	virtual bool ShouldShowInNewMenu() const override;
+};
diff --git a/Engine/Plugins/VirtualProduction/Rivermax/RivermaxMedia/Source/RivermaxMediaEditor/Private/RivermaxMediaSourceFactoryNew.cpp b/Engine/Plugins/VirtualProduction/Rivermax/RivermaxMedia/Source/RivermaxMediaEditor/Private/RivermaxMediaSourceFactoryNew.cpp
new file mode 100644
--- /dev/null
+++ b/Engine/Plugins/VirtualProduction/Rivermax/RivermaxMedia/Source/RivermaxMediaEditor/Private/RivermaxMediaSourceFactoryNew.cpp
@@ -0,0 +1,50 @@
+// Copyright Epic Games, Inc. All Rights Reserved.
+
+#include "RivermaxMediaSourceFactoryNew.h"
+
+#include "AssetTypeCategories.h"
+#include "EngineAnalytics.h"
+#include "RivermaxMediaSource.h"
+
+
+/* URivermaxMediaSourceFactoryNew structors
+ *****************************************************************************/
+
+URivermaxMediaSourceFactoryNew::URivermaxMediaSourceFactoryNew()
+{
+	SupportedClass = URivermaxMediaSource::StaticClass();
+	bCreateNew = true;
+	bEditAfterNew = true;
+}
+
+
+/* UFactory overrides
+ *****************************************************************************/
+
+ /**
+ * @EventName MediaFramework.RivermaxMediaSourceCreated
+ * @Trigger Triggered when a Rivermax media source asset is created.
+ * @Type Client
+ * @Owner MediaIO Team
+ */
+UObject* URivermaxMediaSourceFactoryNew::FactoryCreateNew(UClass* InClass, UObject* InParent, FName InName, EObjectFlags Flags, UObject* Context, FFeedbackContext* Warn)
+{
+	if (FEngineAnalytics::IsAvailable())
+	{
+		FEngineAnalytics::GetProvider().RecordEvent(TEXT("MediaFramework.RivermaxMediaSourceCreated"));
+	}
+
+	return NewObject<URivermaxMediaSource>(InParent, InClass, InName, Flags);
+}
+
+
+uint32 URivermaxMediaSourceFactoryNew::GetMenuCategories() const
+{
+	return EAssetTypeCategories::Media;
+}
+
+
+bool URivermaxMediaSourceFactoryNew::ShouldShowInNewMenu() const
+{
+	return true;
+}
diff --git a/Engine/Plugins/VirtualProduction/Rivermax/RivermaxMedia/Source/RivermaxMediaEditor/Private/RivermaxMediaSourceFactoryNew.h b/Engine/Plugins/VirtualProduction/Rivermax/RivermaxMedia/Source/RivermaxMediaEditor/Private/RivermaxMediaSourceFactoryNew.h
new file mode 100644
--- /dev/null
+++ b/Engine/Plugins/VirtualProduction/Rivermax/RivermaxMedia/Source/RivermaxMediaEditor/Private/RivermaxMediaSourceFactoryNew.h
@@ -0,0 +1,27 @@
+// Copyright Epic Games, Inc. All Rights Reserved.
+
+#pragma once
+
+#include "Factories/Factory.h"
+#include "RivermaxMediaSourceFactoryNew.generated.h"
+
+
+/**
+ * Implements a factory for URivermaxMediaSource objects.
+ */
+UCLASS(hidecategories=Object)
+class URivermaxMediaSourceFactoryNew
+	: public UFactory
+{
+	GENERATED_BODY()
+
+public:
+
+	URivermaxMediaSourceFactoryNew();
+
+	//~ UFactory Interface
+
+	virtual UObject* FactoryCreateNew(UClass* InClass, UObject* InParent, FName InName, EObjectFlags Flags, UObject* Context, FFeedbackContext* Warn) override;
+	virtual uint32 GetMenuCategories() const override;
+	virtual bool ShouldShowInNewMenu() const override;
+};
diff --git a/Engine/Plugins/VirtualProduction/Rivermax/RivermaxMedia/Source/RivermaxMediaEditor/RivermaxMediaEditor.build.cs b/Engine/Plugins/VirtualProduction/Rivermax/RivermaxMedia/Source/RivermaxMediaEditor/RivermaxMediaEditor.build.cs
new file mode 100644
--- /dev/null
+++ b/Engine/Plugins/VirtualProduction/Rivermax/RivermaxMedia/Source/RivermaxMediaEditor/RivermaxMediaEditor.build.cs
@@ -0,0 +1,29 @@
+// Copyright Epic Games, Inc. All Rights Reserved.
+
+using UnrealBuildTool;
+
+public class RivermaxMediaEditor : ModuleRules
+{
+	public RivermaxMediaEditor(ReadOnlyTargetRules Target) : base(Target)
+	{
+		PublicDependencyModuleNames.AddRange(
+			new string[] {
+			}
+		);
+
+		PrivateDependencyModuleNames.AddRange(
+			new string[] {
+				"AssetTools",
+				"Core",
+				"CoreUObject",
+				"Engine",
+				"MediaAssets",
+				"MediaIOCore",
+				"MediaIOEditor",
+				"RivermaxMedia",
+				"UnrealEd"
+
+			}
+		);
+	}
+}
diff --git a/Engine/Plugins/VirtualProduction/Rivermax/RivermaxMedia/Source/RivermaxMediaFactory/Private/RivermaxMediaFactoryModule.cpp b/Engine/Plugins/VirtualProduction/Rivermax/RivermaxMedia/Source/RivermaxMediaFactory/Private/RivermaxMediaFactoryModule.cpp
new file mode 100644
--- /dev/null
+++ b/Engine/Plugins/VirtualProduction/Rivermax/RivermaxMedia/Source/RivermaxMediaFactory/Private/RivermaxMediaFactoryModule.cpp
@@ -0,0 +1,140 @@
+// Copyright Epic Games, Inc. All Rights Reserved.
+
+#include "CoreMinimal.h"
+
+#include "IMediaPlayerFactory.h"
+#include "IMediaModule.h"
+#include "Modules/ModuleInterface.h"
+#include "Modules/ModuleManager.h"
+
+#if WITH_EDITOR
+	#include "ISettingsModule.h"
+	#include "ISettingsSection.h"
+#endif
+
+#include "IRivermaxMediaModule.h"
+
+
+#define LOCTEXT_NAMESPACE "RivermaxMediaFactoryModule"
+
+
+/**
+ * Implements the RivermaxMediaFactory module.
+ */
+class FRivermaxMediaFactoryModule
+	: public IMediaPlayerFactory
+	, public IModuleInterface
+{
+public:
+
+	//~ IMediaPlayerFactory interface
+
+	virtual bool CanPlayUrl(const FString& Url, const IMediaOptions* /*Options*/, TArray<FText>* /*OutWarnings*/, TArray<FText>* OutErrors) const override
+	{
+		FString Scheme;
+		FString Location;
+
+		// check scheme
+		if (!Url.Split(TEXT("://"), &Scheme, &Location, ESearchCase::CaseSensitive))
+		{
+			if (OutErrors != nullptr)
+			{
+				OutErrors->Add(LOCTEXT("NoSchemeFound", "No URI scheme found"));
+			}
+
+			return false;
+		}
+
+		if (!SupportedUriSchemes.Contains(Scheme))
+		{
+			if (OutErrors != nullptr)
+			{
+				OutErrors->Add(FText::Format(LOCTEXT("SchemeNotSupported", "The URI scheme '{0}' is not supported"), FText::FromString(Scheme)));
+			}
+
+			return false;
+		}
+
+		return true;
+	}
+
+	virtual TSharedPtr<IMediaPlayer, ESPMode::ThreadSafe> CreatePlayer(IMediaEventSink& EventSink) override
+	{
+		IRivermaxMediaModule* RivermaxMediaModule = FModuleManager::LoadModulePtr<IRivermaxMediaModule>("RivermaxMedia");
+		return (RivermaxMediaModule != nullptr) ? RivermaxMediaModule->CreatePlayer(EventSink) : nullptr;
+	}
+
+	virtual FText GetDisplayName() const override
+	{
+		return LOCTEXT("MediaPlayerDisplayName", "Rivermax Device Interface");
+	}
+
+	virtual FName GetPlayerName() const override
+	{
+		static FName PlayerName(TEXT("RivermaxMedia"));
+		return PlayerName;
+	}
+
+	virtual FGuid GetPlayerPluginGUID() const override
+	{
+		static FGuid PlayerPluginGUID(0xF537595A, 0x8E8D452B, 0xB8C05707, 0x6B334234);
+		return PlayerPluginGUID;
+	}
+
+	virtual const TArray<FString>& GetSupportedPlatforms() const override
+	{
+		return SupportedPlatforms;
+	}
+
+	virtual bool SupportsFeature(EMediaFeature Feature) const override
+	{
+		return Feature == EMediaFeature::AudioSamples ||
+				Feature == EMediaFeature::MetadataTracks ||
+				Feature == EMediaFeature::VideoSamples;
+
+	}
+
+public:
+
+	//~ IModuleInterface interface
+
+	virtual void StartupModule() override
+	{
+		// supported platforms
+		SupportedPlatforms.Add(TEXT("Windows"));
+
+		// supported schemes
+		SupportedUriSchemes.Add(TEXT("rmax"));
+
+		// register player factory
+		IMediaModule* MediaModule = FModuleManager::LoadModulePtr<IMediaModule>("Media");
+
+		if (MediaModule != nullptr)
+		{
+			MediaModule->RegisterPlayerFactory(*this);
+		}
+	}
+
+	virtual void ShutdownModule() override
+	{
+		// unregister player factory
+		IMediaModule* MediaModule = FModuleManager::GetModulePtr<IMediaModule>("Media");
+
+		if (MediaModule != nullptr)
+		{
+			MediaModule->UnregisterPlayerFactory(*this);
+		}
+	}
+
+private:
+	/** List of platforms that the media player support. */
+	TArray<FString> SupportedPlatforms;
+
+	/** List of supported URI schemes. */
+	TArray<FString> SupportedUriSchemes;
+};
+
+
+#undef LOCTEXT_NAMESPACE
+
+IMPLEMENT_MODULE(FRivermaxMediaFactoryModule, RivermaxMediaFactory);
diff --git a/Engine/Plugins/VirtualProduction/Rivermax/RivermaxMedia/Source/RivermaxMediaFactory/RivermaxMediaFactory.Build.cs b/Engine/Plugins/VirtualProduction/Rivermax/RivermaxMedia/Source/RivermaxMediaFactory/RivermaxMediaFactory.Build.cs
new file mode 100644
--- /dev/null
+++ b/Engine/Plugins/VirtualProduction/Rivermax/RivermaxMedia/Source/RivermaxMediaFactory/RivermaxMediaFactory.Build.cs
@@ -0,0 +1,24 @@
+// Copyright Epic Games, Inc. All Rights Reserved.
+
+namespace UnrealBuildTool.Rules
+{
+	public class RivermaxMediaFactory : ModuleRules
+	{
+		public RivermaxMediaFactory(ReadOnlyTargetRules Target) : base(Target)
+		{
+			PrivateDependencyModuleNames.AddRange(
+				new string[] {
+					"Core",
+					"CoreUObject",
+					"Media",
+					"MediaAssets",
+					"Projects",
+					"RivermaxMedia"
+				});
+
+			PublicDependencyModuleNames.AddRange(
+				new string[] {
+				});
+		}
+	}
+}
