- Fixing rmax CIS static analysis

diff --git a/Engine/Plugins/VirtualProduction/Rivermax/RivermaxMedia/Source/RivermaxMedia/Private/RivermaxMediaCapture.cpp b/Engine/Plugins/VirtualProduction/Rivermax/RivermaxMedia/Source/RivermaxMedia/Private/RivermaxMediaCapture.cpp
--- a/Engine/Plugins/VirtualProduction/Rivermax/RivermaxMedia/Source/RivermaxMedia/Private/RivermaxMediaCapture.cpp
+++ b/Engine/Plugins/VirtualProduction/Rivermax/RivermaxMedia/Source/RivermaxMedia/Private/RivermaxMediaCapture.cpp
@@ -55,7 +55,7 @@ namespace UE::RivermaxMediaCaptureUtil
 		// We need to adjust horizontal resolution if it's not aligned to pgroup pixel coverage
 		// RTP header has to describe a full pgroup even if it's outside of effective resolution.
 		const uint32 PixelAlignment = InFormatInfo.PixelGroupCoverage;
-		const uint32 AlignedHorizontalResolution = InSize.X % PixelAlignment ? InSize.X + (PixelAlignment - InSize.X % PixelAlignment) : InSize.X;
+		const uint32 AlignedHorizontalResolution = (InSize.X % PixelAlignment) ? InSize.X + (PixelAlignment - (InSize.X % PixelAlignment)) : InSize.X;
 		return FIntPoint(AlignedHorizontalResolution, InSize.Y);
 	}
 
diff --git a/Engine/Plugins/VirtualProduction/Rivermax/RivermaxMedia/Source/RivermaxMedia/Private/RivermaxMediaPlayer.cpp b/Engine/Plugins/VirtualProduction/Rivermax/RivermaxMedia/Source/RivermaxMedia/Private/RivermaxMediaPlayer.cpp
--- a/Engine/Plugins/VirtualProduction/Rivermax/RivermaxMedia/Source/RivermaxMedia/Private/RivermaxMediaPlayer.cpp
+++ b/Engine/Plugins/VirtualProduction/Rivermax/RivermaxMedia/Source/RivermaxMedia/Private/RivermaxMediaPlayer.cpp
@@ -307,7 +307,7 @@ namespace UE::RivermaxMedia
 		StreamOptions.PixelFormat = UE::RivermaxMediaUtils::Private::MediaSourcePixelFormatToRivermaxSamplingType(DesiredPixelFormat);
 		const FVideoFormatInfo FormatInfo = FStandardVideoFormat::GetVideoFormatInfo(StreamOptions.PixelFormat);
 		const uint32 PixelAlignment = FormatInfo.PixelGroupCoverage;
-		const uint32 AlignedHorizontalResolution = StreamOptions.Resolution.X % PixelAlignment ? StreamOptions.Resolution.X + (PixelAlignment - StreamOptions.Resolution.X % PixelAlignment) : StreamOptions.Resolution.X;
+		const uint32 AlignedHorizontalResolution = (StreamOptions.Resolution.X % PixelAlignment) ? StreamOptions.Resolution.X + (PixelAlignment - (StreamOptions.Resolution.X % PixelAlignment)) : StreamOptions.Resolution.X;
 		StreamOptions.AlignedResolution = FIntPoint(AlignedHorizontalResolution, StreamOptions.Resolution.Y);
 
 		return true;
