- Fixing traces when doing rivermax output/input
- Removing log spam when frames are missed

diff --git a/Engine/Plugins/Media/MediaIOFramework/Source/MediaIOCore/Private/MediaCapture.cpp b/Engine/Plugins/Media/MediaIOFramework/Source/MediaIOCore/Private/MediaCapture.cpp
--- a/Engine/Plugins/Media/MediaIOFramework/Source/MediaIOCore/Private/MediaCapture.cpp
+++ b/Engine/Plugins/Media/MediaIOFramework/Source/MediaIOCore/Private/MediaCapture.cpp
@@ -48,6 +48,12 @@
 DECLARE_CYCLE_STAT(TEXT("MediaCapture RenderThread FrameCapture"), STAT_MediaCapture_RenderThread_FrameCapture, STATGROUP_Media);
 DECLARE_CYCLE_STAT(TEXT("MediaCapture RenderThread LockResource"), STAT_MediaCapture_RenderThread_LockResource, STATGROUP_Media);
 DECLARE_CYCLE_STAT(TEXT("MediaCapture RenderThread RHI Capture Callback"), STAT_MediaCapture_RenderThread_CaptureCallback, STATGROUP_Media);
+DECLARE_GPU_STAT(MediaCapture_CaptureFrame);
+DECLARE_GPU_STAT(MediaCapture_CustomCapture);
+DECLARE_GPU_STAT(MediaCapture_Conversion);
+DECLARE_GPU_STAT(MediaCapture_Readback);
+DECLARE_GPU_STAT(MediaCapture_ProcessCapture);
+
 
 /** These pixel formats do not require additional conversion except for swizzling and normalized sampling. */
 static TSet<EPixelFormat> SupportedRgbaSwizzleFormats =
@@ -306,7 +312,7 @@ namespace UE::MediaCaptureData
 		/** Registers an external texture to be tracked by the graph and returns a pointer to the tracked resource */
 		FRDGBufferRef RegisterResource(FRDGBuilder& RDGBuilder)
 		{
-			return RDGBuilder.RegisterExternalBuffer(Buffer);
+			return RDGBuilder.RegisterExternalBuffer(Buffer, TEXT("OutputBuffer"));
 		}
 
 		/** Adds a readback pass to the graph */
@@ -631,6 +637,8 @@ namespace UE::MediaCaptureData
 		template<typename TFrameType>
 		static bool CaptureFrame(const FCaptureFrameArgs& Args, TFrameType* CapturingFrame)
 		{
+			RDG_GPU_STAT_SCOPE(Args.GraphBuilder, MediaCapture_CaptureFrame)
+
 			// Validate if we have a resources used to capture source texture
 			if (!CapturingFrame->HasValidResource())
 			{
@@ -683,6 +691,7 @@ namespace UE::MediaCaptureData
 					// If custom conversion was requested from implementation, give it useful information to apply 
 					if (Args.MediaCapture->ConversionOperation == EMediaCaptureConversionOperation::CUSTOM)
 					{
+						RDG_GPU_STAT_SCOPE(Args.GraphBuilder, MediaCapture_CustomCapture)
 						TRACE_CPUPROFILER_EVENT_SCOPE(UMediaCapture::CustomCapture);
 
 						Args.MediaCapture->OnCustomCapture_RenderingThread(Args.GraphBuilder, CapturingFrame->CaptureBaseData, CapturingFrame->UserData
@@ -690,6 +699,7 @@ namespace UE::MediaCaptureData
 					}
 					else
 					{
+						RDG_GPU_STAT_SCOPE(Args.GraphBuilder, MediaCapture_Conversion)
 						TRACE_CPUPROFILER_EVENT_SCOPE(UMediaCapture::FormatConversion);
 						AddConversionPass(Args, { SourceRGBTexture, bRequiresFormatConversion,  CopyInfo, SizeU, SizeV }, OutputResource);
 					}
@@ -698,6 +708,7 @@ namespace UE::MediaCaptureData
 				// If Capture implementation is not grabbing GPU resource directly, push a readback pass to access it from CPU
 				if (Args.MediaCapture->bShouldCaptureRHIResource == false)
 				{
+					RDG_GPU_STAT_SCOPE(Args.GraphBuilder, MediaCapture_Readback)
 					TRACE_CPUPROFILER_EVENT_SCOPE(UMediaCapture::EnqueueReadback);
 
 					CapturingFrame->EnqueueCopy(Args.GraphBuilder, OutputResource);
@@ -1721,6 +1732,7 @@ bool UMediaCapture::ProcessCapture_RenderThread(FRDGBuilder& GraphBuilder, UMedi
 {
 	using namespace UE::MediaCaptureData;
 	
+	RDG_GPU_STAT_SCOPE(GraphBuilder, MediaCapture_ProcessCapture)
 	TRACE_CPUPROFILER_EVENT_SCOPE(UMediaCapture::ProcessCapture_RenderThread);
 
 	if (CapturingFrame)
diff --git a/Engine/Plugins/VirtualProduction/Rivermax/RivermaxCore/Source/RivermaxCore/Private/RivermaxOutputStream.cpp b/Engine/Plugins/VirtualProduction/Rivermax/RivermaxCore/Source/RivermaxCore/Private/RivermaxOutputStream.cpp
--- a/Engine/Plugins/VirtualProduction/Rivermax/RivermaxCore/Source/RivermaxCore/Private/RivermaxOutputStream.cpp
+++ b/Engine/Plugins/VirtualProduction/Rivermax/RivermaxCore/Source/RivermaxCore/Private/RivermaxOutputStream.cpp
@@ -822,7 +822,7 @@ namespace UE::RivermaxCore::Private
 			else
 			{
 				// No frame to send, keep last one and restarts its internal counters
-				UE_LOG(LogRivermax, Warning, TEXT("No frame to send. Reusing last frame '%d'"), CurrentFrame->FrameIndex);
+				UE_LOG(LogRivermax, Verbose, TEXT("No frame to send. Reusing last frame '%d'"), CurrentFrame->FrameIndex);
 				InitializeNextFrame(CurrentFrame);		
 				
 				// Since we want to resend last frame, we need to fast forward chunk pointer to re-point to the one we just sent
diff --git a/Engine/Plugins/VirtualProduction/Rivermax/RivermaxMedia/Source/RivermaxMedia/Private/RivermaxMediaCapture.cpp b/Engine/Plugins/VirtualProduction/Rivermax/RivermaxMedia/Source/RivermaxMedia/Private/RivermaxMediaCapture.cpp
--- a/Engine/Plugins/VirtualProduction/Rivermax/RivermaxMedia/Source/RivermaxMedia/Private/RivermaxMediaCapture.cpp
+++ b/Engine/Plugins/VirtualProduction/Rivermax/RivermaxMedia/Source/RivermaxMedia/Private/RivermaxMediaCapture.cpp
@@ -20,6 +20,9 @@
 #endif
 
 
+DECLARE_GPU_STAT(Rivermax_Capture);
+
+
 
 /* namespace RivermaxMediaCaptureDevice
 *****************************************************************************/
@@ -250,9 +253,11 @@ void URivermaxMediaCapture::OnFrameCaptured_RenderingThread(const FCaptureBaseDa
 	}
 }
 
-void URivermaxMediaCapture::OnRHIResourceCaptured_RenderingThread(const FCaptureBaseData& InBaseData,	TSharedPtr<FMediaCaptureUserData, ESPMode::ThreadSafe> InUserData, FTextureRHIRef InTexture)
+void URivermaxMediaCapture::OnRHIResourceCaptured_RenderingThread(const FCaptureBaseData& InBaseData,	TSharedPtr<FMediaCaptureUserData, ESPMode::ThreadSafe> InUserData, FBufferRHIRef InBuffer)
 {
-	TRACE_CPUPROFILER_EVENT_SCOPE(URivermaxMediaCapture::OnFrameCaptured_RenderingThread);
+	using namespace UE::RivermaxCore;
+
+	TRACE_CPUPROFILER_EVENT_SCOPE(URivermaxMediaCapture::OnRHIResourceCaptured_RenderingThread);
 }
 
 FIntPoint URivermaxMediaCapture::GetCustomOutputSize(const FIntPoint& InSize) const
@@ -292,11 +297,17 @@ FRDGBufferDesc URivermaxMediaCapture::GetCustomBufferDescription(const FIntPoint
 	uint32 BytesPerElement = 0;
 	uint32 ElementsPerRow = 0;
 	UE::RivermaxMediaCaptureUtil::GetOutputEncodingInfo(RivermaxOutput->PixelFormat, InDesiredSize, BytesPerElement, ElementsPerRow);
-	return FRDGBufferDesc::CreateStructuredDesc(BytesPerElement, ElementsPerRow * InDesiredSize.Y);;
+	FRDGBufferDesc Desc = FRDGBufferDesc::CreateStructuredDesc(BytesPerElement, ElementsPerRow * InDesiredSize.Y);
+	
+	// Required when GPUDirect using CUDA will be involved
+	Desc.Usage |= EBufferUsageFlags::Shared;
+	return Desc;
 }
 
 void URivermaxMediaCapture::OnCustomCapture_RenderingThread(FRDGBuilder& GraphBuilder, const FCaptureBaseData& InBaseData, TSharedPtr<FMediaCaptureUserData, ESPMode::ThreadSafe> InUserData, FRDGTextureRef InSourceTexture, FRDGBufferRef OutputBuffer, const FRHICopyTextureInfo& CopyInfo, FVector2D CropU, FVector2D CropV)
 {
+	RDG_GPU_STAT_SCOPE(GraphBuilder, Rivermax_Capture)
+
 	using namespace UE::RivermaxShaders;
 	URivermaxMediaOutput* RivermaxOutput = CastChecked<URivermaxMediaOutput>(MediaOutput);
 
diff --git a/Engine/Plugins/VirtualProduction/Rivermax/RivermaxMedia/Source/RivermaxMedia/Private/RivermaxMediaTextureSampleConverter.cpp b/Engine/Plugins/VirtualProduction/Rivermax/RivermaxMedia/Source/RivermaxMedia/Private/RivermaxMediaTextureSampleConverter.cpp
--- a/Engine/Plugins/VirtualProduction/Rivermax/RivermaxMedia/Source/RivermaxMedia/Private/RivermaxMediaTextureSampleConverter.cpp
+++ b/Engine/Plugins/VirtualProduction/Rivermax/RivermaxMedia/Source/RivermaxMedia/Private/RivermaxMediaTextureSampleConverter.cpp
@@ -7,6 +7,9 @@
 #include "RivermaxMediaTextureSample.h"
 
 
+DECLARE_GPU_STAT(RivermaxSource_SampleConversion);
+
+
 void FRivermaxMediaTextureSampleConverter::Setup(ERivermaxMediaSourcePixelFormat InPixelFormat, TWeakPtr<FRivermaxMediaTextureSample> InSample, bool bInDoSRGBToLinear)
 {
 	InputPixelFormat = InPixelFormat;
@@ -29,193 +32,196 @@ bool FRivermaxMediaTextureSampleConverter::Convert(FTexture2DRHIRef& InDestinati
 	FIntVector GroupCount;
 	
 	FRDGBuilder GraphBuilder(FRHICommandListExecutor::GetImmediateCommandList());
-	SCOPED_DRAW_EVENT(GraphBuilder.RHICmdList, Rivermax_SampleConverter);
-	
-	FRDGTextureRef OutputResource = GraphBuilder.RegisterExternalTexture(CreateRenderTarget(InDestinationTexture, TEXT("RivermaxMediaTextureOutputResource")));
-
-	//Configure shader and add conversion pass based on desired pixel format
-	switch (InputPixelFormat)
-	{
-	case ERivermaxMediaSourcePixelFormat::YUV422_8bit:
 	{
-		TRACE_CPUPROFILER_EVENT_SCOPE(RivermaxSampleConverter::YUV8ShaderSetup);
-		const int32 BytesPerElement = sizeof(FYUV8Bit422ToRGBACS::FYUV8Bit422Buffer);
-		const uint32 Stride = SamplePtr->GetStride();
-		const int32 ElementsPerRow = (Stride / BytesPerElement) + ((Stride % BytesPerElement > 0) ? 1 : 0);
-		const int32 ElementCount = ElementsPerRow * InDestinationTexture->GetDesc().Extent.Y;
-
-		FYUV8Bit422ToRGBACS::FPermutationDomain PermutationVector;
-		PermutationVector.Set<FYUV8Bit422ToRGBACS::FSRGBToLinear>(bDoSRGBToLinear);
-
-		FRDGBufferRef InputYUVBuffer = CreateStructuredBuffer(GraphBuilder, TEXT("RivermaxInputBuffer"), BytesPerElement, ElementCount, SamplePtr->GetBuffer(), BytesPerElement * ElementCount);
-		constexpr int32 PixelsPerInput = 2;
-		const FIntPoint ProcessedOutputDimension = { InDestinationTexture->GetDesc().Extent.X / PixelsPerInput, InDestinationTexture->GetDesc().Extent.Y };
-		GroupCount = FComputeShaderUtils::GetGroupCount(ProcessedOutputDimension, FComputeShaderUtils::kGolden2DGroupSize);
-		FGlobalShaderMap* GlobalShaderMap = GetGlobalShaderMap(GMaxRHIFeatureLevel);
-		TShaderMapRef<FYUV8Bit422ToRGBACS> ComputeShader(GlobalShaderMap, PermutationVector);
-		FMatrix YUVToRGBMatrix = SamplePtr->GetYUVToRGBMatrix();
-		FVector YUVOffset(MediaShaders::YUVOffset8bits);
-		FYUV8Bit422ToRGBACS::FParameters* Parameters = ComputeShader->AllocateAndSetParameters(GraphBuilder, InputYUVBuffer, OutputResource, YUVToRGBMatrix, YUVOffset, ElementsPerRow, ProcessedOutputDimension.Y);
-
-
-		FComputeShaderUtils::AddPass(
-			GraphBuilder
-			, RDG_EVENT_NAME("YUV8Bit422ToRGBA")
-			, ComputeShader
-			, Parameters
-			, GroupCount);
-		break;
-	}
-	case ERivermaxMediaSourcePixelFormat::YUV422_10bit:
-	{
-		TRACE_CPUPROFILER_EVENT_SCOPE(RivermaxSampleConverter::YUV10ShaderSetup);
-		const int32 BytesPerElement = sizeof(FYUV10Bit422ToRGBACS::FYUV10Bit422LEBuffer);
-		const uint32 Stride = SamplePtr->GetStride();
-		const int32 ElementsPerRow = (Stride / BytesPerElement) + ((Stride % BytesPerElement > 0) ? 1 : 0);
-		const int32 ElementCount = ElementsPerRow * InDestinationTexture->GetDesc().Extent.Y;
-
-		FYUV10Bit422ToRGBACS::FPermutationDomain PermutationVector;
-		PermutationVector.Set<FYUV10Bit422ToRGBACS::FSRGBToLinear>(bDoSRGBToLinear);
-
-		FRDGBufferRef InputYUVBuffer = CreateStructuredBuffer(GraphBuilder, TEXT("RivermaxInputBuffer"), BytesPerElement, ElementCount, SamplePtr->GetBuffer(), BytesPerElement * ElementCount);
-		constexpr int32 PixelsPerInput = 8;
-		const FIntPoint ProcessedOutputDimension = { InDestinationTexture->GetDesc().Extent.X / PixelsPerInput, InDestinationTexture->GetDesc().Extent.Y };
-		GroupCount = FComputeShaderUtils::GetGroupCount(ProcessedOutputDimension, FComputeShaderUtils::kGolden2DGroupSize);
-		FGlobalShaderMap* GlobalShaderMap = GetGlobalShaderMap(GMaxRHIFeatureLevel);
-		TShaderMapRef<FYUV10Bit422ToRGBACS> ComputeShader(GlobalShaderMap, PermutationVector);
-		FMatrix YUVToRGBMatrix = SamplePtr->GetYUVToRGBMatrix();
-		FVector YUVOffset(MediaShaders::YUVOffset10bits);
-		FYUV10Bit422ToRGBACS::FParameters* Parameters = ComputeShader->AllocateAndSetParameters(GraphBuilder, InputYUVBuffer, OutputResource, YUVToRGBMatrix, YUVOffset, ElementsPerRow, ProcessedOutputDimension.Y);
-
-
-		FComputeShaderUtils::AddPass(
-			GraphBuilder
-			, RDG_EVENT_NAME("YUV10Bit422ToRGBA")
-			, ComputeShader
-			, Parameters
-			, GroupCount);
-		break;
-	}
-	case ERivermaxMediaSourcePixelFormat::RGB_8bit:
-	{
-		TRACE_CPUPROFILER_EVENT_SCOPE(RivermaxSampleConverter::RGB8ShaderSetup);
-		const int32 BytesPerElement = sizeof(FRGB8BitToRGBA8CS::FRGB8BitBuffer);
-		const uint32 Stride = SamplePtr->GetStride();
-		const int32 ElementsPerRow = (Stride / BytesPerElement) + ((Stride % BytesPerElement > 0) ? 1 : 0);
-		const int32 ElementCount = ElementsPerRow * InDestinationTexture->GetDesc().Extent.Y;
-
-		FRGB8BitToRGBA8CS::FPermutationDomain PermutationVector;
-		PermutationVector.Set<FRGB8BitToRGBA8CS::FSRGBToLinear>(bDoSRGBToLinear);
-
-		FRDGBufferRef InputRGGBuffer = CreateStructuredBuffer(GraphBuilder, TEXT("RivermaxInputBuffer"), BytesPerElement, ElementCount, SamplePtr->GetBuffer(), BytesPerElement * ElementCount);
-		constexpr int32 PixelsPerInput = 4;
-		const FIntPoint ProcessedOutputDimension = { InDestinationTexture->GetDesc().Extent.X / PixelsPerInput,InDestinationTexture->GetDesc().Extent.Y };
-		GroupCount = FComputeShaderUtils::GetGroupCount(ProcessedOutputDimension, FComputeShaderUtils::kGolden2DGroupSize);
-		FGlobalShaderMap* GlobalShaderMap = GetGlobalShaderMap(GMaxRHIFeatureLevel);
-		TShaderMapRef<FRGB8BitToRGBA8CS> ComputeShader(GlobalShaderMap, PermutationVector);
-		FRGB8BitToRGBA8CS::FParameters* Parameters = ComputeShader->AllocateAndSetParameters(GraphBuilder, InputRGGBuffer, OutputResource, ElementsPerRow, ProcessedOutputDimension.Y);
-
-
-		FComputeShaderUtils::AddPass(
-			GraphBuilder
-			, RDG_EVENT_NAME("RGB8BitToRGBA8")
-			, ComputeShader
-			, Parameters
-			, GroupCount);
-		break;
-	}
-	case ERivermaxMediaSourcePixelFormat::RGB_10bit:
-	{
-		TRACE_CPUPROFILER_EVENT_SCOPE(RivermaxSampleConverter::RGB10ShaderSetup);
-
-		const int32 BytesPerElement = sizeof(FRGB10BitToRGBA10CS::FRGB10BitBuffer);
-		const uint32 Stride = SamplePtr->GetStride();
-		const int32 ElementsPerRow = (Stride / BytesPerElement) + ((Stride % BytesPerElement > 0) ? 1 : 0);
-		const int32 ElementCount = ElementsPerRow * InDestinationTexture->GetDesc().Extent.Y;
-		
-		FRGB10BitToRGBA10CS::FPermutationDomain PermutationVector;
-		PermutationVector.Set<FRGB10BitToRGBA10CS::FSRGBToLinear>(bDoSRGBToLinear);
-		
-		FRDGBufferRef InputRGGBuffer = CreateStructuredBuffer(GraphBuilder, TEXT("RivermaxInputBuffer"), BytesPerElement, ElementCount, SamplePtr->GetBuffer(), BytesPerElement * ElementCount);
-		constexpr int32 PixelsPerInput = 16;
-		const FIntPoint ProcessedOutputDimension = { InDestinationTexture->GetDesc().Extent.X / PixelsPerInput,InDestinationTexture->GetDesc().Extent.Y};
-		GroupCount = FComputeShaderUtils::GetGroupCount(ProcessedOutputDimension, FComputeShaderUtils::kGolden2DGroupSize);
-		FGlobalShaderMap* GlobalShaderMap = GetGlobalShaderMap(GMaxRHIFeatureLevel);
-		TShaderMapRef<FRGB10BitToRGBA10CS> ComputeShader(GlobalShaderMap, PermutationVector);
-		FRGB10BitToRGBA10CS::FParameters* Parameters = ComputeShader->AllocateAndSetParameters(GraphBuilder, InputRGGBuffer, OutputResource, ElementsPerRow, ProcessedOutputDimension.Y);
-
-
-		FComputeShaderUtils::AddPass(
-			GraphBuilder
-			, RDG_EVENT_NAME("RGB10BitToRGBA")
-			, ComputeShader
-			, Parameters
-			, GroupCount);
-		break;
-	}
-	case ERivermaxMediaSourcePixelFormat::RGB_12bit:
-	{
-		TRACE_CPUPROFILER_EVENT_SCOPE(RivermaxSampleConverter::RGB12ShaderSetup);
-
-		const int32 BytesPerElement = sizeof(FRGB12BitToRGBA12CS::FRGB12BitBuffer);
-		const uint32 Stride = SamplePtr->GetStride();
-		const int32 ElementsPerRow = (Stride / BytesPerElement) + ((Stride % BytesPerElement > 0) ? 1 : 0);
-		const int32 ElementCount = ElementsPerRow * InDestinationTexture->GetDesc().Extent.Y;
-
-		FRGB12BitToRGBA12CS::FPermutationDomain PermutationVector;
-		PermutationVector.Set<FRGB12BitToRGBA12CS::FSRGBToLinear>(bDoSRGBToLinear);
-
-		FRDGBufferRef InputRGGBuffer = CreateStructuredBuffer(GraphBuilder, TEXT("RivermaxInputBuffer"), BytesPerElement, ElementCount, SamplePtr->GetBuffer(), BytesPerElement * ElementCount);
-		constexpr int32 PixelsPerInput = 8;
-		const FIntPoint ProcessedOutputDimension = { InDestinationTexture->GetDesc().Extent.X / PixelsPerInput,InDestinationTexture->GetDesc().Extent.Y };
-		GroupCount = FComputeShaderUtils::GetGroupCount(ProcessedOutputDimension, FComputeShaderUtils::kGolden2DGroupSize);
-		FGlobalShaderMap* GlobalShaderMap = GetGlobalShaderMap(GMaxRHIFeatureLevel);
-		TShaderMapRef<FRGB12BitToRGBA12CS> ComputeShader(GlobalShaderMap, PermutationVector);
-		FRGB12BitToRGBA12CS::FParameters* Parameters = ComputeShader->AllocateAndSetParameters(GraphBuilder, InputRGGBuffer, OutputResource, ElementsPerRow, ProcessedOutputDimension.Y);
-
-
-		FComputeShaderUtils::AddPass(
-			GraphBuilder
-			, RDG_EVENT_NAME("RGB12BitToRGBA")
-			, ComputeShader
-			, Parameters
-			, GroupCount);
-		break;
-	}
-	case ERivermaxMediaSourcePixelFormat::RGB_16bit_Float:
-	{
-		TRACE_CPUPROFILER_EVENT_SCOPE(RivermaxSampleConverter::RGB16FloatShaderSetup);
-
-		const int32 BytesPerElement = sizeof(FRGB16fBitToRGBA16fCS::FRGB16fBuffer);
-		const uint32 Stride = SamplePtr->GetStride();
-		const int32 ElementsPerRow = (Stride / BytesPerElement) + ((Stride % BytesPerElement > 0) ? 1 : 0);
-		const int32 ElementCount = ElementsPerRow * InDestinationTexture->GetDesc().Extent.Y;
-
-		FRGB16fBitToRGBA16fCS::FPermutationDomain PermutationVector;
-		PermutationVector.Set<FRGB16fBitToRGBA16fCS::FSRGBToLinear>(bDoSRGBToLinear);
-
-		FRDGBufferRef InputRGGBuffer = CreateStructuredBuffer(GraphBuilder, TEXT("RivermaxInputBuffer"), BytesPerElement, ElementCount, SamplePtr->GetBuffer(), BytesPerElement * ElementCount);
-		constexpr int32 PixelsPerInput = 2;
-		const FIntPoint ProcessedOutputDimension = { InDestinationTexture->GetDesc().Extent.X / PixelsPerInput,InDestinationTexture->GetDesc().Extent.Y };
-		GroupCount = FComputeShaderUtils::GetGroupCount(ProcessedOutputDimension, FComputeShaderUtils::kGolden2DGroupSize);
-		FGlobalShaderMap* GlobalShaderMap = GetGlobalShaderMap(GMaxRHIFeatureLevel);
-		TShaderMapRef<FRGB16fBitToRGBA16fCS> ComputeShader(GlobalShaderMap, PermutationVector);
-		FRGB16fBitToRGBA16fCS::FParameters* Parameters = ComputeShader->AllocateAndSetParameters(GraphBuilder, InputRGGBuffer, OutputResource, ElementsPerRow, ProcessedOutputDimension.Y);
-
-
-		FComputeShaderUtils::AddPass(
-			GraphBuilder
-			, RDG_EVENT_NAME("RGB16fBitToRGBA")
-			, ComputeShader
-			, Parameters
-			, GroupCount);
-		break;
-	}
-	default:
-	{
-		ensureMsgf(false, TEXT("Unhandled pixel format (%d) given to Rivermax MediaSample converter"), InputPixelFormat);
-		return false;
-	}
+		RDG_GPU_STAT_SCOPE(GraphBuilder, RivermaxSource_SampleConversion)
+		SCOPED_DRAW_EVENT(GraphBuilder.RHICmdList, Rivermax_SampleConverter);
+
+		FRDGTextureRef OutputResource = GraphBuilder.RegisterExternalTexture(CreateRenderTarget(InDestinationTexture, TEXT("RivermaxMediaTextureOutputResource")));
+
+		//Configure shader and add conversion pass based on desired pixel format
+		switch (InputPixelFormat)
+		{
+		case ERivermaxMediaSourcePixelFormat::YUV422_8bit:
+		{
+			TRACE_CPUPROFILER_EVENT_SCOPE(RivermaxSampleConverter::YUV8ShaderSetup);
+			const int32 BytesPerElement = sizeof(FYUV8Bit422ToRGBACS::FYUV8Bit422Buffer);
+			const uint32 Stride = SamplePtr->GetStride();
+			const int32 ElementsPerRow = (Stride / BytesPerElement) + ((Stride % BytesPerElement > 0) ? 1 : 0);
+			const int32 ElementCount = ElementsPerRow * InDestinationTexture->GetDesc().Extent.Y;
+
+			FYUV8Bit422ToRGBACS::FPermutationDomain PermutationVector;
+			PermutationVector.Set<FYUV8Bit422ToRGBACS::FSRGBToLinear>(bDoSRGBToLinear);
+
+			FRDGBufferRef InputYUVBuffer = CreateStructuredBuffer(GraphBuilder, TEXT("RivermaxInputBuffer"), BytesPerElement, ElementCount, SamplePtr->GetBuffer(), BytesPerElement * ElementCount);
+			constexpr int32 PixelsPerInput = 2;
+			const FIntPoint ProcessedOutputDimension = { InDestinationTexture->GetDesc().Extent.X / PixelsPerInput, InDestinationTexture->GetDesc().Extent.Y };
+			GroupCount = FComputeShaderUtils::GetGroupCount(ProcessedOutputDimension, FComputeShaderUtils::kGolden2DGroupSize);
+			FGlobalShaderMap* GlobalShaderMap = GetGlobalShaderMap(GMaxRHIFeatureLevel);
+			TShaderMapRef<FYUV8Bit422ToRGBACS> ComputeShader(GlobalShaderMap, PermutationVector);
+			FMatrix YUVToRGBMatrix = SamplePtr->GetYUVToRGBMatrix();
+			FVector YUVOffset(MediaShaders::YUVOffset8bits);
+			FYUV8Bit422ToRGBACS::FParameters* Parameters = ComputeShader->AllocateAndSetParameters(GraphBuilder, InputYUVBuffer, OutputResource, YUVToRGBMatrix, YUVOffset, ElementsPerRow, ProcessedOutputDimension.Y);
+
+
+			FComputeShaderUtils::AddPass(
+				GraphBuilder
+				, RDG_EVENT_NAME("YUV8Bit422ToRGBA")
+				, ComputeShader
+				, Parameters
+				, GroupCount);
+			break;
+		}
+		case ERivermaxMediaSourcePixelFormat::YUV422_10bit:
+		{
+			TRACE_CPUPROFILER_EVENT_SCOPE(RivermaxSampleConverter::YUV10ShaderSetup);
+			const int32 BytesPerElement = sizeof(FYUV10Bit422ToRGBACS::FYUV10Bit422LEBuffer);
+			const uint32 Stride = SamplePtr->GetStride();
+			const int32 ElementsPerRow = (Stride / BytesPerElement) + ((Stride % BytesPerElement > 0) ? 1 : 0);
+			const int32 ElementCount = ElementsPerRow * InDestinationTexture->GetDesc().Extent.Y;
+
+			FYUV10Bit422ToRGBACS::FPermutationDomain PermutationVector;
+			PermutationVector.Set<FYUV10Bit422ToRGBACS::FSRGBToLinear>(bDoSRGBToLinear);
+
+			FRDGBufferRef InputYUVBuffer = CreateStructuredBuffer(GraphBuilder, TEXT("RivermaxInputBuffer"), BytesPerElement, ElementCount, SamplePtr->GetBuffer(), BytesPerElement * ElementCount);
+			constexpr int32 PixelsPerInput = 8;
+			const FIntPoint ProcessedOutputDimension = { InDestinationTexture->GetDesc().Extent.X / PixelsPerInput, InDestinationTexture->GetDesc().Extent.Y };
+			GroupCount = FComputeShaderUtils::GetGroupCount(ProcessedOutputDimension, FComputeShaderUtils::kGolden2DGroupSize);
+			FGlobalShaderMap* GlobalShaderMap = GetGlobalShaderMap(GMaxRHIFeatureLevel);
+			TShaderMapRef<FYUV10Bit422ToRGBACS> ComputeShader(GlobalShaderMap, PermutationVector);
+			FMatrix YUVToRGBMatrix = SamplePtr->GetYUVToRGBMatrix();
+			FVector YUVOffset(MediaShaders::YUVOffset10bits);
+			FYUV10Bit422ToRGBACS::FParameters* Parameters = ComputeShader->AllocateAndSetParameters(GraphBuilder, InputYUVBuffer, OutputResource, YUVToRGBMatrix, YUVOffset, ElementsPerRow, ProcessedOutputDimension.Y);
+
+
+			FComputeShaderUtils::AddPass(
+				GraphBuilder
+				, RDG_EVENT_NAME("YUV10Bit422ToRGBA")
+				, ComputeShader
+				, Parameters
+				, GroupCount);
+			break;
+		}
+		case ERivermaxMediaSourcePixelFormat::RGB_8bit:
+		{
+			TRACE_CPUPROFILER_EVENT_SCOPE(RivermaxSampleConverter::RGB8ShaderSetup);
+			const int32 BytesPerElement = sizeof(FRGB8BitToRGBA8CS::FRGB8BitBuffer);
+			const uint32 Stride = SamplePtr->GetStride();
+			const int32 ElementsPerRow = (Stride / BytesPerElement) + ((Stride % BytesPerElement > 0) ? 1 : 0);
+			const int32 ElementCount = ElementsPerRow * InDestinationTexture->GetDesc().Extent.Y;
+
+			FRGB8BitToRGBA8CS::FPermutationDomain PermutationVector;
+			PermutationVector.Set<FRGB8BitToRGBA8CS::FSRGBToLinear>(bDoSRGBToLinear);
+
+			FRDGBufferRef InputRGGBuffer = CreateStructuredBuffer(GraphBuilder, TEXT("RivermaxInputBuffer"), BytesPerElement, ElementCount, SamplePtr->GetBuffer(), BytesPerElement * ElementCount);
+			constexpr int32 PixelsPerInput = 4;
+			const FIntPoint ProcessedOutputDimension = { InDestinationTexture->GetDesc().Extent.X / PixelsPerInput,InDestinationTexture->GetDesc().Extent.Y };
+			GroupCount = FComputeShaderUtils::GetGroupCount(ProcessedOutputDimension, FComputeShaderUtils::kGolden2DGroupSize);
+			FGlobalShaderMap* GlobalShaderMap = GetGlobalShaderMap(GMaxRHIFeatureLevel);
+			TShaderMapRef<FRGB8BitToRGBA8CS> ComputeShader(GlobalShaderMap, PermutationVector);
+			FRGB8BitToRGBA8CS::FParameters* Parameters = ComputeShader->AllocateAndSetParameters(GraphBuilder, InputRGGBuffer, OutputResource, ElementsPerRow, ProcessedOutputDimension.Y);
+
+
+			FComputeShaderUtils::AddPass(
+				GraphBuilder
+				, RDG_EVENT_NAME("RGB8BitToRGBA8")
+				, ComputeShader
+				, Parameters
+				, GroupCount);
+			break;
+		}
+		case ERivermaxMediaSourcePixelFormat::RGB_10bit:
+		{
+			TRACE_CPUPROFILER_EVENT_SCOPE(RivermaxSampleConverter::RGB10ShaderSetup);
+
+			const int32 BytesPerElement = sizeof(FRGB10BitToRGBA10CS::FRGB10BitBuffer);
+			const uint32 Stride = SamplePtr->GetStride();
+			const int32 ElementsPerRow = (Stride / BytesPerElement) + ((Stride % BytesPerElement > 0) ? 1 : 0);
+			const int32 ElementCount = ElementsPerRow * InDestinationTexture->GetDesc().Extent.Y;
+
+			FRGB10BitToRGBA10CS::FPermutationDomain PermutationVector;
+			PermutationVector.Set<FRGB10BitToRGBA10CS::FSRGBToLinear>(bDoSRGBToLinear);
+
+			FRDGBufferRef InputRGGBuffer = CreateStructuredBuffer(GraphBuilder, TEXT("RivermaxInputBuffer"), BytesPerElement, ElementCount, SamplePtr->GetBuffer(), BytesPerElement * ElementCount);
+			constexpr int32 PixelsPerInput = 16;
+			const FIntPoint ProcessedOutputDimension = { InDestinationTexture->GetDesc().Extent.X / PixelsPerInput,InDestinationTexture->GetDesc().Extent.Y };
+			GroupCount = FComputeShaderUtils::GetGroupCount(ProcessedOutputDimension, FComputeShaderUtils::kGolden2DGroupSize);
+			FGlobalShaderMap* GlobalShaderMap = GetGlobalShaderMap(GMaxRHIFeatureLevel);
+			TShaderMapRef<FRGB10BitToRGBA10CS> ComputeShader(GlobalShaderMap, PermutationVector);
+			FRGB10BitToRGBA10CS::FParameters* Parameters = ComputeShader->AllocateAndSetParameters(GraphBuilder, InputRGGBuffer, OutputResource, ElementsPerRow, ProcessedOutputDimension.Y);
+
+
+			FComputeShaderUtils::AddPass(
+				GraphBuilder
+				, RDG_EVENT_NAME("RGB10BitToRGBA")
+				, ComputeShader
+				, Parameters
+				, GroupCount);
+			break;
+		}
+		case ERivermaxMediaSourcePixelFormat::RGB_12bit:
+		{
+			TRACE_CPUPROFILER_EVENT_SCOPE(RivermaxSampleConverter::RGB12ShaderSetup);
+
+			const int32 BytesPerElement = sizeof(FRGB12BitToRGBA12CS::FRGB12BitBuffer);
+			const uint32 Stride = SamplePtr->GetStride();
+			const int32 ElementsPerRow = (Stride / BytesPerElement) + ((Stride % BytesPerElement > 0) ? 1 : 0);
+			const int32 ElementCount = ElementsPerRow * InDestinationTexture->GetDesc().Extent.Y;
+
+			FRGB12BitToRGBA12CS::FPermutationDomain PermutationVector;
+			PermutationVector.Set<FRGB12BitToRGBA12CS::FSRGBToLinear>(bDoSRGBToLinear);
+
+			FRDGBufferRef InputRGGBuffer = CreateStructuredBuffer(GraphBuilder, TEXT("RivermaxInputBuffer"), BytesPerElement, ElementCount, SamplePtr->GetBuffer(), BytesPerElement * ElementCount);
+			constexpr int32 PixelsPerInput = 8;
+			const FIntPoint ProcessedOutputDimension = { InDestinationTexture->GetDesc().Extent.X / PixelsPerInput,InDestinationTexture->GetDesc().Extent.Y };
+			GroupCount = FComputeShaderUtils::GetGroupCount(ProcessedOutputDimension, FComputeShaderUtils::kGolden2DGroupSize);
+			FGlobalShaderMap* GlobalShaderMap = GetGlobalShaderMap(GMaxRHIFeatureLevel);
+			TShaderMapRef<FRGB12BitToRGBA12CS> ComputeShader(GlobalShaderMap, PermutationVector);
+			FRGB12BitToRGBA12CS::FParameters* Parameters = ComputeShader->AllocateAndSetParameters(GraphBuilder, InputRGGBuffer, OutputResource, ElementsPerRow, ProcessedOutputDimension.Y);
+
+
+			FComputeShaderUtils::AddPass(
+				GraphBuilder
+				, RDG_EVENT_NAME("RGB12BitToRGBA")
+				, ComputeShader
+				, Parameters
+				, GroupCount);
+			break;
+		}
+		case ERivermaxMediaSourcePixelFormat::RGB_16bit_Float:
+		{
+			TRACE_CPUPROFILER_EVENT_SCOPE(RivermaxSampleConverter::RGB16FloatShaderSetup);
+
+			const int32 BytesPerElement = sizeof(FRGB16fBitToRGBA16fCS::FRGB16fBuffer);
+			const uint32 Stride = SamplePtr->GetStride();
+			const int32 ElementsPerRow = (Stride / BytesPerElement) + ((Stride % BytesPerElement > 0) ? 1 : 0);
+			const int32 ElementCount = ElementsPerRow * InDestinationTexture->GetDesc().Extent.Y;
+
+			FRGB16fBitToRGBA16fCS::FPermutationDomain PermutationVector;
+			PermutationVector.Set<FRGB16fBitToRGBA16fCS::FSRGBToLinear>(bDoSRGBToLinear);
+
+			FRDGBufferRef InputRGGBuffer = CreateStructuredBuffer(GraphBuilder, TEXT("RivermaxInputBuffer"), BytesPerElement, ElementCount, SamplePtr->GetBuffer(), BytesPerElement * ElementCount);
+			constexpr int32 PixelsPerInput = 2;
+			const FIntPoint ProcessedOutputDimension = { InDestinationTexture->GetDesc().Extent.X / PixelsPerInput,InDestinationTexture->GetDesc().Extent.Y };
+			GroupCount = FComputeShaderUtils::GetGroupCount(ProcessedOutputDimension, FComputeShaderUtils::kGolden2DGroupSize);
+			FGlobalShaderMap* GlobalShaderMap = GetGlobalShaderMap(GMaxRHIFeatureLevel);
+			TShaderMapRef<FRGB16fBitToRGBA16fCS> ComputeShader(GlobalShaderMap, PermutationVector);
+			FRGB16fBitToRGBA16fCS::FParameters* Parameters = ComputeShader->AllocateAndSetParameters(GraphBuilder, InputRGGBuffer, OutputResource, ElementsPerRow, ProcessedOutputDimension.Y);
+
+
+			FComputeShaderUtils::AddPass(
+				GraphBuilder
+				, RDG_EVENT_NAME("RGB16fBitToRGBA")
+				, ComputeShader
+				, Parameters
+				, GroupCount);
+			break;
+		}
+		default:
+		{
+			ensureMsgf(false, TEXT("Unhandled pixel format (%d) given to Rivermax MediaSample converter"), InputPixelFormat);
+			return false;
+		}
+		}
 	}
 
 	GraphBuilder.Execute();
diff --git a/Engine/Plugins/VirtualProduction/Rivermax/RivermaxMedia/Source/RivermaxMedia/Public/RivermaxMediaCapture.h b/Engine/Plugins/VirtualProduction/Rivermax/RivermaxMedia/Source/RivermaxMedia/Public/RivermaxMediaCapture.h
--- a/Engine/Plugins/VirtualProduction/Rivermax/RivermaxMedia/Source/RivermaxMedia/Public/RivermaxMediaCapture.h
+++ b/Engine/Plugins/VirtualProduction/Rivermax/RivermaxMedia/Source/RivermaxMedia/Public/RivermaxMediaCapture.h
@@ -40,8 +40,8 @@ protected:
 	
 	virtual void BeforeFrameCaptured_RenderingThread(const FCaptureBaseData& InBaseData, TSharedPtr<FMediaCaptureUserData, ESPMode::ThreadSafe> InUserData, FTextureRHIRef InTexture) override;
 	virtual void OnFrameCaptured_RenderingThread(const FCaptureBaseData& InBaseData, TSharedPtr<FMediaCaptureUserData, ESPMode::ThreadSafe> InUserData, void* InBuffer, int32 Width, int32 Height, int32 BytesPerRow) override;
-	virtual void OnRHIResourceCaptured_RenderingThread(const FCaptureBaseData& InBaseData, TSharedPtr<FMediaCaptureUserData, ESPMode::ThreadSafe> InUserData, FTextureRHIRef InTexture) override;
-	
+	virtual void OnRHIResourceCaptured_RenderingThread(const FCaptureBaseData& InBaseData, TSharedPtr<FMediaCaptureUserData, ESPMode::ThreadSafe> InUserData, FBufferRHIRef InBuffer) override;
+
 	/** For custom conversion, methods that need to be overriden */
 	virtual FIntPoint GetCustomOutputSize(const FIntPoint& InSize) const override;
 	virtual EMediaCaptureResourceType GetCustomOutputResourceType() const override;
