- Fixing missing api to leverage gpudirect in rivermax output stream

diff --git a/Engine/Plugins/VirtualProduction/Rivermax/RivermaxCore/Source/RivermaxCore/Private/RivermaxManager.cpp b/Engine/Plugins/VirtualProduction/Rivermax/RivermaxCore/Source/RivermaxCore/Private/RivermaxManager.cpp
--- a/Engine/Plugins/VirtualProduction/Rivermax/RivermaxCore/Source/RivermaxCore/Private/RivermaxManager.cpp
+++ b/Engine/Plugins/VirtualProduction/Rivermax/RivermaxCore/Source/RivermaxCore/Private/RivermaxManager.cpp
@@ -205,6 +205,11 @@ namespace UE::RivermaxCore::Private
 	{
 		return DeviceFinder->IsValidIP(InSourceIP);
 	}
+
+	bool FRivermaxManager::IsGPUDirectSupported() const
+	{
+		return bIsGPUDirectSupported;
+	}
 }
 
 
diff --git a/Engine/Plugins/VirtualProduction/Rivermax/RivermaxCore/Source/RivermaxCore/Private/RivermaxManager.h b/Engine/Plugins/VirtualProduction/Rivermax/RivermaxCore/Source/RivermaxCore/Private/RivermaxManager.h
--- a/Engine/Plugins/VirtualProduction/Rivermax/RivermaxCore/Source/RivermaxCore/Private/RivermaxManager.h
+++ b/Engine/Plugins/VirtualProduction/Rivermax/RivermaxCore/Source/RivermaxCore/Private/RivermaxManager.h
@@ -22,19 +22,24 @@ namespace UE::RivermaxCore::Private
 		virtual TConstArrayView<FRivermaxDeviceInfo> GetDevices() const override;
 		virtual bool GetMatchingDevice(const FString& InSourceIP, FString& OutDeviceIP) const override;
 		virtual bool IsValidIP(const FString& InSourceIP) const override;
+		virtual bool IsGPUDirectSupported() const;
 		//~ End IRivermaxManager interface
 
 	private:
 		bool LoadRivermaxLibrary();
 		void InitializeLibrary();
-		void FindDevices();
 
 	private:
 
 		/** Whether the library has been initialized and is usable */
 		bool bIsInitialized = false;
+
+		/** Handle pointer to the loaded library */
 		void* LibraryHandle = nullptr;
 
+		/** Whether GPU direct is supported globally for rivermax. Currently means Cuda is present with RDMA support */
+		bool bIsGPUDirectSupported = false;
+
 		/** Current time source */
 		ERivermaxTimeSource TimeSource = ERivermaxTimeSource::PTP;
 
diff --git a/Engine/Plugins/VirtualProduction/Rivermax/RivermaxCore/Source/RivermaxCore/Private/RivermaxOutputStream.cpp b/Engine/Plugins/VirtualProduction/Rivermax/RivermaxCore/Source/RivermaxCore/Private/RivermaxOutputStream.cpp
--- a/Engine/Plugins/VirtualProduction/Rivermax/RivermaxCore/Source/RivermaxCore/Private/RivermaxOutputStream.cpp
+++ b/Engine/Plugins/VirtualProduction/Rivermax/RivermaxCore/Source/RivermaxCore/Private/RivermaxOutputStream.cpp
@@ -109,6 +109,8 @@ namespace UE::RivermaxCore::Private
 
 	bool FRivermaxOutputStream::Initialize(const FRivermaxStreamOptions& InOptions, IRivermaxOutputStreamListener& InListener)
 	{
+		TRACE_CPUPROFILER_EVENT_SCOPE(FRivermaxOutputStream::Initialize);
+
 		IRivermaxCoreModule& RivermaxModule = FModuleManager::LoadModuleChecked<IRivermaxCoreModule>(TEXT("RivermaxCore"));
 		if (RivermaxModule.GetRivermaxManager()->IsInitialized() == false)
 		{
@@ -127,6 +129,12 @@ namespace UE::RivermaxCore::Private
 			return false;
 		}
 
+		// We do (try) to make gpu allocations here to let the capturer know if we require it or not.
+		if (RivermaxModule.GetRivermaxManager()->IsGPUDirectSupported() && Options.bUseGPUDirect)
+		{
+			//bUseGPUDirect = AllocateGPUBuffers();
+		}
+
 		Async(EAsyncExecution::TaskGraph, [this]()
 		{
 			//Create event to wait on when no frames are available to send
@@ -135,11 +143,14 @@ namespace UE::RivermaxCore::Private
 			TAnsiStringBuilder<2048> SDPDescription;
 			UE::RivermaxCore::Private::Utils::StreamOptionsToSDPDescription(Options, SDPDescription);
 
-			// Initialize video buffer memory we will be using
-			InitializeBuffers();
+			// Initialize buffers in system memory if we're not using gpudirect
+			if(bUseGPUDirect == false)
+			{
+				AllocateSystemBuffers();
+			}
 
 			// Initialize internal memory and rivermax configuration 
-			if(InitializeMemory())
+			if(InitializeStreamMemoryConfig())
 			{
 				// Initialize rivermax output stream with desired config
 				const uint32 NumberPacketsPerFrame = StreamMemory.PacketsPerFrame;
@@ -205,14 +216,15 @@ namespace UE::RivermaxCore::Private
 			ReadyToSendEvent = nullptr;
 			UE_LOG(LogRivermax, Log, TEXT("Rivermax Output stream has shutdown"));
 		}
-
 	}
 
 	bool FRivermaxOutputStream::PushVideoFrame(const FRivermaxOutputVideoFrameInfo& NewFrame)
 	{
+		TRACE_CPUPROFILER_EVENT_SCOPE(FRivermaxOutputStream::PushVideoFrame);
+
 		if (TSharedPtr<FRivermaxOutputFrame> AvailableFrame = GetNextAvailableFrame(NewFrame.FrameIdentifier))
 		{
-			const int32 Stride = Options.AlignedResolution.X / FormatInfo.PixelGroupCoverage * FormatInfo.PixelGroupSize;
+			const int32 Stride = GetStride();
 			FMemory::Memcpy(AvailableFrame->VideoBuffer, NewFrame.VideoBuffer, NewFrame.Height * Stride);
 
 			AvailableFrame->bIsVideoBufferReady = true;
@@ -299,9 +311,9 @@ namespace UE::RivermaxCore::Private
 		}
 	}
 
-	void FRivermaxOutputStream::InitializeBuffers()
+	void FRivermaxOutputStream::AllocateSystemBuffers()
 	{
-		const int32 Stride = Options.AlignedResolution.X / FormatInfo.PixelGroupCoverage * FormatInfo.PixelGroupSize;
+		const int32 Stride = GetStride();
 		const int32 FrameAllocSize = Options.AlignedResolution.Y * Stride;
 		AvailableFrames.Reserve(Options.NumberOfBuffers);
 		for (int32 Index = 0; Index < Options.NumberOfBuffers; ++Index)
@@ -314,14 +326,14 @@ namespace UE::RivermaxCore::Private
 		}
 	}
 
-	bool FRivermaxOutputStream::InitializeMemory()
+	bool FRivermaxOutputStream::InitializeStreamMemoryConfig()
 	{
 		using namespace UE::RivermaxCore::Private::Utils;
 
 		//2110 stream type
 		//We need to use the fullframe allocated size to compute the payload size.
 
-		const int32 BytesPerLine = Options.AlignedResolution.X / FormatInfo.PixelGroupCoverage * FormatInfo.PixelGroupSize;
+		const int32 BytesPerLine = GetStride();
 		const uint32 EffectiveBytesPerLine = BytesPerLine;
 
 		const bool bFoundPayload = FindPayloadSize(Options, BytesPerLine, FormatInfo, StreamMemory.PayloadSize);
@@ -393,6 +405,8 @@ namespace UE::RivermaxCore::Private
 
 	TSharedPtr<FRivermaxOutputFrame> FRivermaxOutputStream::GetNextFrameToSend()
 	{
+		TRACE_CPUPROFILER_EVENT_SCOPE(FRivermaxOutputStream::GetNextFrameToSend);
+
 		TSharedPtr<FRivermaxOutputFrame> FrameToSend;
 
 		FScopeLock Lock(&FrameCriticalSection);
@@ -424,6 +438,8 @@ namespace UE::RivermaxCore::Private
 
 	TSharedPtr<FRivermaxOutputFrame> FRivermaxOutputStream::GetNextAvailableFrame(uint32 InFrameIdentifier)
 	{
+		TRACE_CPUPROFILER_EVENT_SCOPE(FRivermaxOutputStream::GetNextAvailableFrame);
+
 		TSharedPtr<FRivermaxOutputFrame> NextFrame;
 
 		{
@@ -942,5 +958,23 @@ namespace UE::RivermaxCore::Private
 			}
 		}
 	}
+
+	bool FRivermaxOutputStream::PushGPUVideoFrame(const FRivermaxOutputVideoFrameInfo& NewFrame, FBufferRHIRef CapturedBuffer)
+	{
+		TRACE_CPUPROFILER_EVENT_SCOPE(FRivermaxOutputStream::PushGPUVideoFrame);
+
+		return false;
+	}
+
+	bool FRivermaxOutputStream::IsGPUDirectSupported() const
+	{
+		return bUseGPUDirect;
+	}
+
+	int32 FRivermaxOutputStream::GetStride() const
+	{
+		check(FormatInfo.PixelGroupCoverage != 0);
+		return (Options.AlignedResolution.X / FormatInfo.PixelGroupCoverage) * FormatInfo.PixelGroupSize;
+	}
 }
 
diff --git a/Engine/Plugins/VirtualProduction/Rivermax/RivermaxCore/Source/RivermaxCore/Private/RivermaxOutputStream.h b/Engine/Plugins/VirtualProduction/Rivermax/RivermaxCore/Source/RivermaxCore/Private/RivermaxOutputStream.h
--- a/Engine/Plugins/VirtualProduction/Rivermax/RivermaxCore/Source/RivermaxCore/Private/RivermaxOutputStream.h
+++ b/Engine/Plugins/VirtualProduction/Rivermax/RivermaxCore/Source/RivermaxCore/Private/RivermaxOutputStream.h
@@ -98,6 +98,8 @@ namespace UE::RivermaxCore::Private
 		virtual bool Initialize(const FRivermaxStreamOptions& Options, IRivermaxOutputStreamListener& InListener) override;
 		virtual void Uninitialize() override;
 		virtual bool PushVideoFrame(const FRivermaxOutputVideoFrameInfo& NewFrame) override;
+		virtual bool PushGPUVideoFrame(const FRivermaxOutputVideoFrameInfo& NewFrame, FBufferRHIRef CapturedBuffer) override;
+		virtual bool IsGPUDirectSupported() const override;
 		//~ End IRivermaxOutputStream interface
 
 		void Process_AnyThread();
@@ -110,8 +112,7 @@ namespace UE::RivermaxCore::Private
 		//~ End FRunnable interface
 
 	private:
-		void InitializeBuffers();
-		bool InitializeMemory();
+		bool InitializeStreamMemoryConfig();
 		void InitializeNextFrame(const TSharedPtr<FRivermaxOutputFrame>& NextFrame);
 		TSharedPtr<FRivermaxOutputFrame> GetNextFrameToSend();
 		TSharedPtr<FRivermaxOutputFrame> GetNextAvailableFrame(uint32 InFrameIdentifier);
@@ -125,30 +126,55 @@ namespace UE::RivermaxCore::Private
 		void InitializeStreamTimingSettings();
 		void ShowStats();
 		uint32 GetTimestampFromTime(uint64 InTimeNanosec, double InMediaClockRate) const;
+		void AllocateSystemBuffers();
+		int32 GetStride() const;
 
 	private:
+
+		/** Options related to this stream. i.e resolution, frame rate, etc... */
 		FRivermaxStreamOptions Options;
+
+		/** Rivermax memory configuration. i.e. memblock, chunks */
 		FRivermaxOutputStreamMemory StreamMemory;
+
+		/** Various stats collected by this stream */
 		FRivermaxOutputStreamStats Stats;
+
+		/** State of various piece for this stream. Alignment points, schedule number, etc... */
 		FRivermaxOutputStreamData StreamData;
 
+		/** Stream id returned by rmax library */
 		rmax_stream_id StreamId;
+
+		/** Critical section to protect frames access */
 		FCriticalSection FrameCriticalSection;
 
+		/** Current frame being sent */
 		TSharedPtr<FRivermaxOutputFrame> CurrentFrame;
 
+		/** Available frames to write memory to (Ready to be written) */
 		TArray<TSharedPtr<FRivermaxOutputFrame>> AvailableFrames;
+
+		/** Frames ready to be sent to rivermax (Ready to be read) */
 		TArray<TSharedPtr<FRivermaxOutputFrame>> FramesToSend;
 
+		/** Thread scheduling frame output */
 		TUniquePtr<FRunnableThread> RivermaxThread;
+
+		/** Whether stream is active or not */
 		std::atomic<bool> bIsActive;
 
+		/** Event used to let scheduler that a frame is ready to be sent */
 		FEvent* ReadyToSendEvent = nullptr;
 
+		/** Listener for this stream events */
 		IRivermaxOutputStreamListener* Listener = nullptr;
 
-		static constexpr double MediaClockSampleRate = 90000.0; //Required to comply with SMTPE 2110-10.The Media Clock and RTP Clock rate for streams compliant to this standard shall be 90 kHz.
-		ERivermaxStreamType StreamType = ERivermaxStreamType::VIDEO_2110_20_STREAM; //todo
+		/** Required to comply with SMTPE 2110 - 10.The Media Clock and RTP Clock rate for streams compliant to this standard shall be 90 kHz. */
+		static constexpr double MediaClockSampleRate = 90000.0; 
+
+		/** Type of stream created. Only 21110-20 (Video is supported now) */
+		ERivermaxStreamType StreamType = ERivermaxStreamType::VIDEO_2110_20_STREAM;
 
 		/** TRoffset time calculated based on ST2110 - 21 Gapped(for now) method. This is added to next alignment point */
 		uint64 TransmitOffsetNanosec = 0;
@@ -158,6 +184,9 @@ namespace UE::RivermaxCore::Private
 
 		/** Timestamp at which we logged stats */
 		double LastStatsShownTimestamp = 0.0;
+		
+		/** Whether stream is using gpudirect to host memory consumed by Rivermax */
+		bool bUseGPUDirect = false;
 	};
 }
 
diff --git a/Engine/Plugins/VirtualProduction/Rivermax/RivermaxCore/Source/RivermaxCore/Public/IRivermaxManager.h b/Engine/Plugins/VirtualProduction/Rivermax/RivermaxCore/Source/RivermaxCore/Public/IRivermaxManager.h
--- a/Engine/Plugins/VirtualProduction/Rivermax/RivermaxCore/Source/RivermaxCore/Public/IRivermaxManager.h
+++ b/Engine/Plugins/VirtualProduction/Rivermax/RivermaxCore/Source/RivermaxCore/Public/IRivermaxManager.h
@@ -46,6 +46,9 @@ namespace UE::RivermaxCore
 
 		/** Returns true if provided IP string is in a valid format */
 		virtual bool IsValidIP(const FString& InSourceIP) const = 0;
+
+		/** Whether gpudirect is supported on this platform. At the moment, we rely on CUDA to handle it. */
+		virtual bool IsGPUDirectSupported() const = 0;
 	};
 }
 
diff --git a/Engine/Plugins/VirtualProduction/Rivermax/RivermaxCore/Source/RivermaxCore/Public/IRivermaxOutputStream.h b/Engine/Plugins/VirtualProduction/Rivermax/RivermaxCore/Source/RivermaxCore/Public/IRivermaxOutputStream.h
--- a/Engine/Plugins/VirtualProduction/Rivermax/RivermaxCore/Source/RivermaxCore/Public/IRivermaxOutputStream.h
+++ b/Engine/Plugins/VirtualProduction/Rivermax/RivermaxCore/Source/RivermaxCore/Public/IRivermaxOutputStream.h
@@ -4,6 +4,8 @@
 
 #include "CoreMinimal.h"
 
+#include "RHI.h"
+
 
 namespace UE::RivermaxCore
 {
@@ -47,6 +49,12 @@ namespace UE::RivermaxCore
 
 		/** Pushes new video frame to the stream. Returns false if frame couldn't be pushed. */
 		virtual bool PushVideoFrame(const FRivermaxOutputVideoFrameInfo& NewFrame) = 0;
+		
+		/** Pushes a frame to be captured from GPU memory */
+		virtual bool PushGPUVideoFrame(const FRivermaxOutputVideoFrameInfo& NewFrame, FBufferRHIRef CapturedBuffer) = 0;
+
+		/** Returns true if GPUDirect is supported */
+		virtual bool IsGPUDirectSupported() const = 0;
 	};
 }
 
diff --git a/Engine/Plugins/VirtualProduction/Rivermax/RivermaxCore/Source/RivermaxCore/Public/RivermaxTypes.h b/Engine/Plugins/VirtualProduction/Rivermax/RivermaxCore/Source/RivermaxCore/Public/RivermaxTypes.h
--- a/Engine/Plugins/VirtualProduction/Rivermax/RivermaxCore/Source/RivermaxCore/Public/RivermaxTypes.h
+++ b/Engine/Plugins/VirtualProduction/Rivermax/RivermaxCore/Source/RivermaxCore/Public/RivermaxTypes.h
@@ -34,6 +34,9 @@ namespace UE::RivermaxCore
 
 		/** Resolution aligning with pgroup of sampling type */
 		FIntPoint AlignedResolution = FIntPoint::ZeroValue;
+
+		/** Whether to leverage GPUDirect (Cuda) capability to transfer memory to NIC if available */
+		bool bUseGPUDirect = false;
 	};
 
 	enum class RIVERMAXCORE_API ERivermaxStreamType : uint8
diff --git a/Engine/Plugins/VirtualProduction/Rivermax/RivermaxCore/Source/RivermaxCore/RivermaxCore.build.cs b/Engine/Plugins/VirtualProduction/Rivermax/RivermaxCore/Source/RivermaxCore/RivermaxCore.build.cs
--- a/Engine/Plugins/VirtualProduction/Rivermax/RivermaxCore/Source/RivermaxCore/RivermaxCore.build.cs
+++ b/Engine/Plugins/VirtualProduction/Rivermax/RivermaxCore/Source/RivermaxCore/RivermaxCore.build.cs
@@ -14,6 +14,7 @@ public class RivermaxCore : ModuleRules
 				"MediaAssets",
 				"MediaIOCore",
 				"RivermaxLib",
+				"RHI",
 				"TimeManagement"
 			}
 		);
@@ -23,7 +24,6 @@ public class RivermaxCore : ModuleRules
 				"Engine",
 				"Networking",
 				"RenderCore",
-				"RHI"
 			}
 		);
 
diff --git a/Engine/Plugins/VirtualProduction/Rivermax/RivermaxMedia/Source/RivermaxMedia/Private/RivermaxMediaCapture.cpp b/Engine/Plugins/VirtualProduction/Rivermax/RivermaxMedia/Source/RivermaxMedia/Private/RivermaxMediaCapture.cpp
--- a/Engine/Plugins/VirtualProduction/Rivermax/RivermaxMedia/Source/RivermaxMedia/Private/RivermaxMediaCapture.cpp
+++ b/Engine/Plugins/VirtualProduction/Rivermax/RivermaxMedia/Source/RivermaxMedia/Private/RivermaxMediaCapture.cpp
@@ -10,7 +10,6 @@
 #include "RivermaxMediaOutput.h"
 #include "RivermaxMediaUtils.h"
 #include "RivermaxShaders.h"
-#include "RivermaxTypes.h"
 #include "Slate/SceneViewport.h"
 #include "Widgets/SViewport.h"
 
@@ -21,9 +20,36 @@
 
 
 DECLARE_GPU_STAT(Rivermax_Capture);
+DECLARE_GPU_STAT(Rivermax_SyncPointPass);
+
+
+/** Structure holding data used for synchronization of buffer output */
+struct URivermaxMediaCapture::FRivermaxCaptureSyncData
+{
+	FRivermaxCaptureSyncData()
+		: bIsBusy(false)
+	{
+	}
+
+	/**
+	 * We use a RHIFence to write to in a pass following the buffer conversion one.
+	 * We spawn a task that will poll for the fence to be over (Waiting is not exposed)
+	 * and then we will push the captured buffer to rivermax stream
+	 */
+	FGPUFenceRHIRef RHIFence;
+
+	/** Whether fence is currently waiting to be cleared */
+	std::atomic<bool> bIsBusy;
+};
 
 
 
+/** Parameter to make our sync pass needing the convert pass as a prereq */
+BEGIN_SHADER_PARAMETER_STRUCT(FRivermaxSyncPassParameters, )
+	RDG_BUFFER_ACCESS(Buffer, ERHIAccess::CopySrc)
+END_SHADER_PARAMETER_STRUCT()
+
+
 /* namespace RivermaxMediaCaptureDevice
 *****************************************************************************/
 
@@ -155,6 +181,8 @@ bool URivermaxMediaCapture::UpdateRenderTargetImpl(UTextureRenderTarget2D* InRen
 
 void URivermaxMediaCapture::StopCaptureImpl(bool bAllowPendingFrameToBeProcess)
 {
+	bIsActive = false;
+
 	if (RivermaxStream)
 	{
 		RivermaxStream->Uninitialize();
@@ -164,16 +192,12 @@ void URivermaxMediaCapture::StopCaptureImpl(bool bAllowPendingFrameToBeProcess)
 
 bool URivermaxMediaCapture::ShouldCaptureRHIResource() const
 {
-	//todo for gpudirect support
-	return false;  
-}
-
-void URivermaxMediaCapture::BeforeFrameCaptured_RenderingThread(const FCaptureBaseData& InBaseData, TSharedPtr<FMediaCaptureUserData, ESPMode::ThreadSafe> InUserData, FTextureRHIRef InTexture)
-{
-	if (ShouldCaptureRHIResource())
+	if (RivermaxStream)
 	{
-	
+		return RivermaxStream->IsGPUDirectSupported();
 	}
+
+	return false;
 }
 
 bool URivermaxMediaCapture::HasFinishedProcessing() const
@@ -186,22 +210,66 @@ bool URivermaxMediaCapture::Initialize(URivermaxMediaOutput* InMediaOutput)
 	using namespace UE::RivermaxCore;
 
 	check(InMediaOutput);
+	bIsActive = false;
 
 	IRivermaxCoreModule* Module = FModuleManager::GetModulePtr<IRivermaxCoreModule>("RivermaxCore");
 	if (Module)
 	{
-		FRivermaxStreamOptions Options;
 		if (ConfigureStream(InMediaOutput, Options))
 		{
 			RivermaxStream = Module->CreateOutputStream();
 			if (RivermaxStream)
 			{
-				return RivermaxStream->Initialize(Options, *this);
+				bAreSyncHandlersInitialized = false;
+				bIsActive = RivermaxStream->Initialize(Options, *this);
 			}
 		}
 	}
 
-	return false;
+	return bIsActive;
+}
+
+void URivermaxMediaCapture::InitializeSyncHandlers_RenderThread()
+{
+	SyncHandlers.Reset(Options.NumberOfBuffers);
+	for (int32 Index = 0; Index < Options.NumberOfBuffers; ++Index)
+	{
+		TSharedPtr<FRivermaxCaptureSyncData> NewSync = MakeShared<FRivermaxCaptureSyncData>();
+		NewSync->RHIFence = RHICreateGPUFence(*FString::Printf(TEXT("RmaxSyncs_%02d"), Index));
+		SyncHandlers.Add(MoveTemp(NewSync));
+	}
+
+	bAreSyncHandlersInitialized = true;
+}
+
+TSharedPtr<URivermaxMediaCapture::FRivermaxCaptureSyncData> URivermaxMediaCapture::GetAvailableSyncHandler() const
+{
+	const auto FindAvailableHandlerFunc = [](const TSharedPtr<FRivermaxCaptureSyncData>& Item)
+	{
+		if (Item->bIsBusy == false)
+		{
+			return true;
+		}
+
+		return false;
+	};
+	
+	if (const TSharedPtr<FRivermaxCaptureSyncData>* FoundItem = SyncHandlers.FindByPredicate(FindAvailableHandlerFunc))
+	{
+		return *FoundItem;
+	}
+
+	return nullptr;
+}
+
+bool URivermaxMediaCapture::AreSyncHandlersBusy() const
+{
+	const auto IsBusyFunc = [](const TSharedPtr<FRivermaxCaptureSyncData>& Item)
+	{
+		return Item->bIsBusy.load();
+	};
+
+	return SyncHandlers.ContainsByPredicate(IsBusyFunc);
 }
 
 bool URivermaxMediaCapture::ConfigureStream(URivermaxMediaOutput* InMediaOutput, UE::RivermaxCore::FRivermaxStreamOptions& OutOptions) const
@@ -227,6 +295,7 @@ bool URivermaxMediaCapture::ConfigureStream(URivermaxMediaOutput* InMediaOutput,
 	OutOptions.Resolution = InMediaOutput->Resolution;
 	OutOptions.FrameRate = InMediaOutput->FrameRate;
 	OutOptions.NumberOfBuffers = InMediaOutput->NumberOfTextureBuffers;
+	OutOptions.bUseGPUDirect = InMediaOutput->bUseGPUDirect;
 
 	OutOptions.PixelFormat = UE::RivermaxMediaUtils::Private::MediaOutputPixelFormatToRivermaxSamplingType(InMediaOutput->PixelFormat);
 	const FVideoFormatInfo Info = FStandardVideoFormat::GetVideoFormatInfo(OutOptions.PixelFormat);
@@ -235,6 +304,73 @@ bool URivermaxMediaCapture::ConfigureStream(URivermaxMediaOutput* InMediaOutput,
 	return true;
 }
 
+void URivermaxMediaCapture::AddSyncPointPass(FRDGBuilder& GraphBuilder, const FCaptureBaseData& InBaseData, FRDGBufferRef OutputBuffer)
+{
+	TRACE_CPUPROFILER_EVENT_SCOPE(RivermaxSyncPoint);
+	RDG_GPU_STAT_SCOPE(GraphBuilder, Rivermax_SyncPointPass);
+
+	// Initialize sync handlers only the first time. 
+	if (bAreSyncHandlersInitialized == false)
+	{
+		InitializeSyncHandlers_RenderThread();
+	}
+
+	// Add buffer output as a parameter to depend on the compute shader pass
+	FRivermaxSyncPassParameters* PassParameters = GraphBuilder.AllocParameters<FRivermaxSyncPassParameters>();
+	PassParameters->Buffer = OutputBuffer;
+
+	// Prepare frame info we are going to push
+	UE::RivermaxCore::FRivermaxOutputVideoFrameInfo NewFrameInfo;
+	NewFrameInfo.FrameIdentifier = InBaseData.SourceFrameNumber;
+
+	URivermaxMediaCapture* Capturer = this;
+	GraphBuilder.AddPass(
+		RDG_EVENT_NAME("RivermaxCopySyncPass"),
+		PassParameters,
+		ERDGPassFlags::Copy | ERDGPassFlags::NeverCull,
+		[Capturer, OutputBuffer, NewFrameInfo](FRHICommandListImmediate& RHICmdList)
+		{
+			if (Capturer->bIsActive)
+			{
+				// Get available sync handler to create a sync point
+				TSharedPtr<FRivermaxCaptureSyncData> SyncDataPtr = Capturer->GetAvailableSyncHandler();
+				if (ensure(SyncDataPtr))
+				{
+					FRHIBuffer* RHIBuffer = OutputBuffer->GetRHI();
+
+					// This will happen after the compute shader conversion pass has completed
+					RHICmdList.WriteGPUFence(SyncDataPtr->RHIFence);
+					SyncDataPtr->bIsBusy = true;
+
+					// Spawn a task that will wait (poll) and continue the process of providing a new buffer to rivermax stream
+					UE::Tasks::Launch(UE_SOURCE_LOCATION, [Capturer, RHIBuffer, NewFrameInfo, SyncDataPtr]()
+					{	
+						TRACE_CPUPROFILER_EVENT_SCOPE(RmaxSyncTask);
+
+						// Wait until fence has been written (shader has completed)
+						while (SyncDataPtr->RHIFence->Poll() == false && Capturer->bIsActive)
+						{
+							FPlatformProcess::SleepNoStats(0);
+						}
+
+						SyncDataPtr->RHIFence->Clear();
+						SyncDataPtr->bIsBusy = false;
+
+						// Push new buffer to rivermax stream
+						if (Capturer->bIsActive && Capturer->RivermaxStream)
+						{
+							if (Capturer->RivermaxStream->PushGPUVideoFrame(NewFrameInfo, RHIBuffer) == false)
+							{
+								UE_LOG(LogRivermaxMedia, Verbose, TEXT("Failed to pushed captured frame"));
+							}
+						}
+
+					}, LowLevelTasks::ETaskPriority::BackgroundHigh);
+				}
+			}
+		});
+}
+
 void URivermaxMediaCapture::OnFrameCaptured_RenderingThread(const FCaptureBaseData& InBaseData, TSharedPtr<FMediaCaptureUserData, ESPMode::ThreadSafe> InUserData, void* InBuffer, int32 Width, int32 Height, int32 BytesPerRow)
 {
 	using namespace UE::RivermaxCore;
@@ -249,17 +385,10 @@ void URivermaxMediaCapture::OnFrameCaptured_RenderingThread(const FCaptureBaseDa
 	NewFrame.FrameIdentifier = InBaseData.SourceFrameNumber;
 	if (RivermaxStream->PushVideoFrame(NewFrame) == false)
 	{
-		UE_LOG(LogRivermaxMedia, VeryVerbose, TEXT("Failed to pushed captured frame"));
+		UE_LOG(LogRivermaxMedia, Verbose, TEXT("Failed to pushed captured frame"));
 	}
 }
 
-void URivermaxMediaCapture::OnRHIResourceCaptured_RenderingThread(const FCaptureBaseData& InBaseData,	TSharedPtr<FMediaCaptureUserData, ESPMode::ThreadSafe> InUserData, FBufferRHIRef InBuffer)
-{
-	using namespace UE::RivermaxCore;
-
-	TRACE_CPUPROFILER_EVENT_SCOPE(URivermaxMediaCapture::OnRHIResourceCaptured_RenderingThread);
-}
-
 FIntPoint URivermaxMediaCapture::GetCustomOutputSize(const FIntPoint& InSize) const
 {
 	URivermaxMediaOutput* RivermaxOutput = CastChecked<URivermaxMediaOutput>(MediaOutput);
@@ -307,6 +436,7 @@ FRDGBufferDesc URivermaxMediaCapture::GetCustomBufferDescription(const FIntPoint
 void URivermaxMediaCapture::OnCustomCapture_RenderingThread(FRDGBuilder& GraphBuilder, const FCaptureBaseData& InBaseData, TSharedPtr<FMediaCaptureUserData, ESPMode::ThreadSafe> InUserData, FRDGTextureRef InSourceTexture, FRDGBufferRef OutputBuffer, const FRHICopyTextureInfo& CopyInfo, FVector2D CropU, FVector2D CropV)
 {
 	RDG_GPU_STAT_SCOPE(GraphBuilder, Rivermax_Capture)
+	TRACE_CPUPROFILER_EVENT_SCOPE(URivermaxMediaCapture::OnCustomCapture_RenderingThread);
 
 	using namespace UE::RivermaxShaders;
 	URivermaxMediaOutput* RivermaxOutput = CastChecked<URivermaxMediaOutput>(MediaOutput);
@@ -406,6 +536,17 @@ void URivermaxMediaCapture::OnCustomCapture_RenderingThread(FRDGBuilder& GraphBu
 		break;
 	}
 	}
+
+	// Adds a simple pass to have a fence written after compute shaders are done
+	if (RivermaxStream->IsGPUDirectSupported())
+	{
+		AddSyncPointPass(GraphBuilder, InBaseData, OutputBuffer);
+	}
+}
+
+bool URivermaxMediaCapture::IsReadyForFinishDestroy()
+{
+	return Super::IsReadyForFinishDestroy() && !AreSyncHandlersBusy();
 }
 
 void URivermaxMediaCapture::OnInitializationCompleted(bool bHasSucceed)
diff --git a/Engine/Plugins/VirtualProduction/Rivermax/RivermaxMedia/Source/RivermaxMedia/Public/RivermaxMediaCapture.h b/Engine/Plugins/VirtualProduction/Rivermax/RivermaxMedia/Source/RivermaxMedia/Public/RivermaxMediaCapture.h
--- a/Engine/Plugins/VirtualProduction/Rivermax/RivermaxMedia/Source/RivermaxMedia/Public/RivermaxMediaCapture.h
+++ b/Engine/Plugins/VirtualProduction/Rivermax/RivermaxMedia/Source/RivermaxMedia/Public/RivermaxMediaCapture.h
@@ -7,6 +7,7 @@
 #include "IRivermaxOutputStream.h"
 #include "Misc/FrameRate.h"
 #include "RivermaxMediaOutput.h"
+#include "RivermaxTypes.h"
 
 #include "RivermaxMediaCapture.generated.h"
 
@@ -37,10 +38,7 @@ protected:
 	virtual bool UpdateRenderTargetImpl(UTextureRenderTarget2D* InRenderTarget) override;
 	virtual void StopCaptureImpl(bool bAllowPendingFrameToBeProcess) override;
 	virtual bool ShouldCaptureRHIResource() const override;
-	
-	virtual void BeforeFrameCaptured_RenderingThread(const FCaptureBaseData& InBaseData, TSharedPtr<FMediaCaptureUserData, ESPMode::ThreadSafe> InUserData, FTextureRHIRef InTexture) override;
 	virtual void OnFrameCaptured_RenderingThread(const FCaptureBaseData& InBaseData, TSharedPtr<FMediaCaptureUserData, ESPMode::ThreadSafe> InUserData, void* InBuffer, int32 Width, int32 Height, int32 BytesPerRow) override;
-	virtual void OnRHIResourceCaptured_RenderingThread(const FCaptureBaseData& InBaseData, TSharedPtr<FMediaCaptureUserData, ESPMode::ThreadSafe> InUserData, FBufferRHIRef InBuffer) override;
 
 	/** For custom conversion, methods that need to be overriden */
 	virtual FIntPoint GetCustomOutputSize(const FIntPoint& InSize) const override;
@@ -49,16 +47,50 @@ protected:
 	virtual void OnCustomCapture_RenderingThread(FRDGBuilder& GraphBuilder, const FCaptureBaseData& InBaseData, TSharedPtr<FMediaCaptureUserData, ESPMode::ThreadSafe> InUserData, FRDGTextureRef InSourceTexture, FRDGBufferRef OutputBuffer, const FRHICopyTextureInfo& CopyInfo, FVector2D CropU, FVector2D CropV) override;
 	//~ End UMediaCapture interface
 
+	//~ Begin UObject interface
+	virtual bool IsReadyForFinishDestroy();
+	//~ End UObject interface
+
 	//~ Begin IRivermaxOutputStreamListener interface
 	virtual void OnInitializationCompleted(bool bHasSucceed) override;
 	virtual void OnStreamError() override;
 	//~ End IRivermaxOutputStreamListener interface
 
 private:
+	struct FRivermaxCaptureSyncData;
+
+	/** Initializes capture and launches stream creation */
 	bool Initialize(URivermaxMediaOutput* InMediaOutput);
+
+	/** Initializes sync handlers, fences, for gpudirect functionality */
+	void InitializeSyncHandlers_RenderThread();
+
+	/** Returns available sync handler. i.e fence not already used */
+	TSharedPtr<FRivermaxCaptureSyncData> GetAvailableSyncHandler() const;
+
+	/** Returns true if one of the sync handler is in flight waiting for its fence */
+	bool AreSyncHandlersBusy() const;
+
+	/** Configures rivermax stream with desired output options */
 	bool ConfigureStream(URivermaxMediaOutput* InMediaOutput, UE::RivermaxCore::FRivermaxStreamOptions& OutOptions) const;
 
+	/** When GPUDirect is involved, we add a pass after the conversion pass to write a fence and know when we can consume the buffer */
+	void AddSyncPointPass(FRDGBuilder& GraphBuilder, const FCaptureBaseData& InBaseData, FRDGBufferRef OutputBuffer);
+
 private:
 
+	/** Instance of the rivermax stream opened for this capture */
 	TUniquePtr<UE::RivermaxCore::IRivermaxOutputStream> RivermaxStream;
+
+	/** Set of options used to configure output stream */
+	UE::RivermaxCore::FRivermaxStreamOptions Options;
+
+	/** Whether sync handlers have been initialized or not.  */
+	bool bAreSyncHandlersInitialized = false;
+
+	/** Whether capture is active. Used to be queried by sync task */
+	std::atomic<bool> bIsActive;
+
+	/** Array of sync handlers (fence) to sync when captured buffer is completed */
+	TArray<TSharedPtr<FRivermaxCaptureSyncData>> SyncHandlers;
 };
diff --git a/Engine/Plugins/VirtualProduction/Rivermax/RivermaxMedia/Source/RivermaxMedia/Public/RivermaxMediaOutput.h b/Engine/Plugins/VirtualProduction/Rivermax/RivermaxMedia/Source/RivermaxMedia/Public/RivermaxMediaOutput.h
--- a/Engine/Plugins/VirtualProduction/Rivermax/RivermaxMedia/Source/RivermaxMedia/Public/RivermaxMediaOutput.h
+++ b/Engine/Plugins/VirtualProduction/Rivermax/RivermaxMedia/Source/RivermaxMedia/Public/RivermaxMediaOutput.h
@@ -51,21 +51,31 @@ public:
 #endif //WITH_EDITOR
 	//~ End UObject interface
 
+	/** Resolution of this output stream */
 	UPROPERTY(EditAnywhere, BlueprintReadWrite, Category="Format")
 	FIntPoint Resolution = {1920, 1080};
 	
+	/** Frame rate of this output stream */
 	UPROPERTY(EditAnywhere, BlueprintReadWrite, Category = "Format")
 	FFrameRate FrameRate = {24,1};
 	
+	/** Pixel format for this output stream */
 	UPROPERTY(EditAnywhere, BlueprintReadWrite, Category = "Format")
 	ERivermaxMediaOutputPixelFormat PixelFormat = ERivermaxMediaOutputPixelFormat::PF_10BIT_RGB;
 
+	/** Address of NIC interface to use. Supports wildcard, i.e. 10.*.69.*. */
 	UPROPERTY(EditAnywhere, BlueprintReadWrite, Category = "Format")
 	FString InterfaceAddress;
 
+	/** Address of the stream. Can be multicast, i.e. 224.1.1.1) */
 	UPROPERTY(EditAnywhere, BlueprintReadWrite, Category = "Format")
 	FString StreamAddress;
 
+	/** Port to use for this output */
 	UPROPERTY(EditAnywhere, BlueprintReadWrite, Category = "Format")
 	int32 Port = 50000;
+
+	/** Whether to use GPUDirect if available (Memcopy from GPU to NIC directly bypassing system) if available */
+	UPROPERTY(EditAnywhere, AdvancedDisplay, Category = "Output")
+	bool bUseGPUDirect = false;
 };
