- Fixing anythread support not working for rivermax output

diff --git a/Engine/Plugins/VirtualProduction/Rivermax/RivermaxMedia/Source/RivermaxMedia/Private/RivermaxMediaCapture.cpp b/Engine/Plugins/VirtualProduction/Rivermax/RivermaxMedia/Source/RivermaxMedia/Private/RivermaxMediaCapture.cpp
--- a/Engine/Plugins/VirtualProduction/Rivermax/RivermaxMedia/Source/RivermaxMedia/Private/RivermaxMediaCapture.cpp
+++ b/Engine/Plugins/VirtualProduction/Rivermax/RivermaxMedia/Source/RivermaxMedia/Private/RivermaxMediaCapture.cpp
@@ -403,6 +403,29 @@ void URivermaxMediaCapture::OnFrameCaptured_RenderingThread(const FCaptureBaseDa
 	}
 }
 
+void URivermaxMediaCapture::OnRHIResourceCaptured_AnyThread(const FCaptureBaseData& InBaseData, TSharedPtr<FMediaCaptureUserData, ESPMode::ThreadSafe> InUserData, FBufferRHIRef InBuffer)
+{
+	// To implement when we start using media capture fence task
+}
+
+void URivermaxMediaCapture::OnFrameCaptured_AnyThread(const FCaptureBaseData& InBaseData, TSharedPtr<FMediaCaptureUserData, ESPMode::ThreadSafe> InUserData, const FMediaCaptureResourceData& InResourceData)
+{
+	using namespace UE::RivermaxCore;
+
+	TRACE_CPUPROFILER_EVENT_SCOPE(URivermaxMediaCapture::OnFrameCaptured_AnyThread);
+	
+	FRivermaxOutputVideoFrameInfo NewFrame;
+	NewFrame.Height = InResourceData.Height;
+	NewFrame.Width = InResourceData.Width;
+	NewFrame.Stride = InResourceData.BytesPerRow;
+	NewFrame.VideoBuffer = InResourceData.Buffer;
+	NewFrame.FrameIdentifier = InBaseData.SourceFrameNumberRenderThread;
+	if (RivermaxStream->PushVideoFrame(NewFrame) == false)
+	{
+		UE_LOG(LogRivermaxMedia, Verbose, TEXT("Failed to pushed captured frame"));
+	}
+}
+
 FIntPoint URivermaxMediaCapture::GetCustomOutputSize(const FIntPoint& InSize) const
 {
 	URivermaxMediaOutput* RivermaxOutput = CastChecked<URivermaxMediaOutput>(MediaOutput);
diff --git a/Engine/Plugins/VirtualProduction/Rivermax/RivermaxMedia/Source/RivermaxMedia/Public/RivermaxMediaCapture.h b/Engine/Plugins/VirtualProduction/Rivermax/RivermaxMedia/Source/RivermaxMedia/Public/RivermaxMediaCapture.h
--- a/Engine/Plugins/VirtualProduction/Rivermax/RivermaxMedia/Source/RivermaxMedia/Public/RivermaxMediaCapture.h
+++ b/Engine/Plugins/VirtualProduction/Rivermax/RivermaxMedia/Source/RivermaxMedia/Public/RivermaxMediaCapture.h
@@ -34,6 +34,12 @@ protected:
 	virtual void StopCaptureImpl(bool bAllowPendingFrameToBeProcess) override;
 	virtual bool ShouldCaptureRHIResource() const override;
 	virtual void OnFrameCaptured_RenderingThread(const FCaptureBaseData& InBaseData, TSharedPtr<FMediaCaptureUserData, ESPMode::ThreadSafe> InUserData, void* InBuffer, int32 Width, int32 Height, int32 BytesPerRow) override;
+	virtual void OnRHIResourceCaptured_AnyThread(const FCaptureBaseData& InBaseData, TSharedPtr<FMediaCaptureUserData, ESPMode::ThreadSafe> InUserData, FBufferRHIRef InBuffer) override;
+	virtual void OnFrameCaptured_AnyThread(const FCaptureBaseData& InBaseData, TSharedPtr<FMediaCaptureUserData, ESPMode::ThreadSafe> InUserData, const FMediaCaptureResourceData& InResourceData) override;
+	virtual bool SupportsAnyThreadCapture() const override
+	{
+		return true;
+	}
 
 	/** For custom conversion, methods that need to be overriden */
 	virtual FIntPoint GetCustomOutputSize(const FIntPoint& InSize) const override;
