- Fix crash when captured resolution is 0x0 in rivermax

diff --git a/Engine/Plugins/VirtualProduction/Rivermax/RivermaxMedia/Source/RivermaxMedia/Private/RivermaxMediaCapture.cpp b/Engine/Plugins/VirtualProduction/Rivermax/RivermaxMedia/Source/RivermaxMedia/Private/RivermaxMediaCapture.cpp
--- a/Engine/Plugins/VirtualProduction/Rivermax/RivermaxMedia/Source/RivermaxMedia/Private/RivermaxMediaCapture.cpp
+++ b/Engine/Plugins/VirtualProduction/Rivermax/RivermaxMedia/Source/RivermaxMedia/Private/RivermaxMediaCapture.cpp
@@ -307,6 +307,13 @@ bool URivermaxMediaCapture::ConfigureStream(URivermaxMediaOutput* InMediaOutput,
 	OutOptions.StreamAddress = InMediaOutput->StreamAddress;
 	OutOptions.Port = InMediaOutput->Port;
 	OutOptions.Resolution = GetDesiredSize();
+
+	if (OutOptions.Resolution.X <= 0 || OutOptions.Resolution.Y <= 0)
+	{
+		UE_LOG(LogRivermaxMedia, Warning, TEXT("Can't start capture. Invalid resolution requested: %dx%d"), OutOptions.Resolution.X, OutOptions.Resolution.Y);
+		return false;
+	}
+
 	OutOptions.FrameRate = InMediaOutput->FrameRate;
 	OutOptions.NumberOfBuffers = InMediaOutput->NumberOfTextureBuffers;
 	OutOptions.bUseGPUDirect = InMediaOutput->bUseGPUDirect;
