- First fix for tearing in 2110 input frame rendering

diff --git a/Engine/Plugins/VirtualProduction/Rivermax/RivermaxMedia/Source/RivermaxMedia/Private/RivermaxMediaPlayer.cpp b/Engine/Plugins/VirtualProduction/Rivermax/RivermaxMedia/Source/RivermaxMedia/Private/RivermaxMediaPlayer.cpp
--- a/Engine/Plugins/VirtualProduction/Rivermax/RivermaxMedia/Source/RivermaxMedia/Private/RivermaxMediaPlayer.cpp
+++ b/Engine/Plugins/VirtualProduction/Rivermax/RivermaxMedia/Source/RivermaxMedia/Private/RivermaxMediaPlayer.cpp
@@ -595,17 +595,24 @@ namespace UE::RivermaxMedia
 			const TSharedPtr<FRivermaxMediaTextureSample>& NextSample = SamplePool[NextFrameExpectations.FrameIndex]->Sample;
 			if (bDoesStreamSupportsGPUDirect)
 			{
-				OutConverterSetup.GPUBuffer = NextSample->GetGPUBuffer();
+				// For gpudirect path, we return the expected buffer directly
+				// PreConvertFunc will take care of waiting on RHI thread for the buffer to be usable
+				OutConverterSetup.GetGPUBufferFunc = [Sample = NextSample](){ return Sample->GetGPUBuffer(); };
+				OutConverterSetup.PreConvertFunc = [NextFrameExpectations, this](const FRDGBuilder& GraphBuilder)
+				{
+					PreSampleUsage(NextFrameExpectations);
+				};
 			}
 			else
 			{
-				OutConverterSetup.SystemBuffer = NextSample->GetBuffer();
+				// For system memory path, we wait for the sample to be received before returning it.
+				// This will stall the render thread in its current state
+				OutConverterSetup.GetSystemBufferFunc = [NextFrameExpectations, Sample = NextSample, this]()
+				{
+					WaitForSample(NextFrameExpectations);
+					return Sample->GetBuffer();
+				};
 			}
-			
-			OutConverterSetup.PreConvertFunc = [NextFrameExpectations, this](FRDGBuilder& GraphBuilder)
-			{
-				PreSampleUsage(GraphBuilder, NextFrameExpectations);
-			};
 
 			// Setup post sample usage pass
 			OutConverterSetup.PostConvertFunc = [NextFrameExpectations, this] (FRDGBuilder& GraphBuilder)
@@ -838,62 +845,15 @@ namespace UE::RivermaxMedia
 		return IMediaPlayer::GetPlayerFeatureFlag(flag);
 	}
 
-	void FRivermaxMediaPlayer::PreSampleUsage(FRDGBuilder& GraphBuilder, const FFrameExpectation& FrameExpectation)
+	void FRivermaxMediaPlayer::PreSampleUsage(const FFrameExpectation& FrameExpectation)
 	{
 		// Before the sample is used during conversion, we want to make sure the expected frame is available to be used
-
-		GraphBuilder.AddPass(RDG_EVENT_NAME("RivermaxPreSampleUsage"),
-			ERDGPassFlags::NeverCull,
+		FRHICommandListImmediate& RHICmdList = FRHICommandListExecutor::GetImmediateCommandList();
+		RHICmdList.EnqueueLambda(
 			[FrameExpectation, this](FRHICommandList& RHICmdList)
 			{
 				TRACE_CPUPROFILER_EVENT_SCOPE(RivermaxWaitForSamples);
-
-				const double StartTimeSeconds = FPlatformTime::Seconds();
-				constexpr double TimeoutSeconds = 0.5;
-
-				while (true)
-				{
-					{
-						if (!IsReadyToPlay())
-						{
-							break;
-						}
-
-						if (FrameExpectation.FrameIndex == INDEX_NONE)
-						{
-							break;
-						}
-
-						if (IsFrameSkipped(FrameExpectation.FrameNumber))
-						{
-							UE_LOG(LogRivermaxMedia, Verbose, TEXT("Stopped waiting for frame %u as it was marked as skipped."), FrameExpectation.FrameNumber);
-							break;
-						}
-
-						// Our goal here is to wait until the expected frame is available to be used (received) unless there is a timeout
-						const TSharedPtr<FRivermaxSampleWrapper>& Frame = SamplePool[FrameExpectation.FrameIndex];
-						if (Frame->ReceptionState == ESampleReceptionState::Received)
-						{
-							if (Frame->FrameNumber >= FrameExpectation.FrameNumber)
-							{
-								if (Frame->FrameNumber > FrameExpectation.FrameNumber)
-								{
-									UE_LOG(LogRivermaxMedia, Warning, TEXT("More recent frame, %u, than expected, %u, found while waiting."), Frame->FrameNumber, FrameExpectation.FrameNumber);
-								}
-								break;
-							}
-						}
-					}
-
-					{
-						FPlatformProcess::SleepNoStats(SleepTimeSeconds);
-						if ((FPlatformTime::Seconds() - StartTimeSeconds) > TimeoutSeconds)
-						{
-							UE_LOG(LogRivermaxMedia, Error, TEXT("Timed out waiting for frame %u."), FrameExpectation.FrameNumber);
-							break;
-						}
-					}
-				}
+				WaitForSample(FrameExpectation);
 			});
 	}
 
@@ -1042,6 +1002,59 @@ namespace UE::RivermaxMedia
 			}
 		}
 	}
+
+	void FRivermaxMediaPlayer::WaitForSample(const FFrameExpectation& FrameExpectation)
+	{
+		TRACE_CPUPROFILER_EVENT_SCOPE(RivermaxWaitForSample);
+
+		const double StartTimeSeconds = FPlatformTime::Seconds();
+		constexpr double TimeoutSeconds = 0.5;
+
+		while (true)
+		{
+			{
+				if (!IsReadyToPlay())
+				{
+					break;
+				}
+
+				if (FrameExpectation.FrameIndex == INDEX_NONE)
+				{
+					break;
+				}
+
+				if (IsFrameSkipped(FrameExpectation.FrameNumber))
+				{
+					UE_LOG(LogRivermaxMedia, Verbose, TEXT("Stopped waiting for frame %u as it was marked as skipped."), FrameExpectation.FrameNumber);
+					break;
+				}
+
+				// Our goal here is to wait until the expected frame is available to be used (received) unless there is a timeout
+				const TSharedPtr<FRivermaxSampleWrapper>& Frame = SamplePool[FrameExpectation.FrameIndex];
+				if (Frame->ReceptionState == ESampleReceptionState::Received)
+				{
+					if (Frame->FrameNumber >= FrameExpectation.FrameNumber)
+					{
+						if (Frame->FrameNumber > FrameExpectation.FrameNumber)
+						{
+							UE_LOG(LogRivermaxMedia, Warning, TEXT("More recent frame, %u, than expected, %u, found while waiting."), Frame->FrameNumber, FrameExpectation.FrameNumber);
+						}
+						break;
+					}
+				}
+			}
+
+			{
+				FPlatformProcess::SleepNoStats(SleepTimeSeconds);
+				if ((FPlatformTime::Seconds() - StartTimeSeconds) > TimeoutSeconds)
+				{
+					UE_LOG(LogRivermaxMedia, Error, TEXT("Timed out waiting for frame %u."), FrameExpectation.FrameNumber);
+					break;
+				}
+			}
+		}
+	}
+
 }
 
 #undef LOCTEXT_NAMESPACE
diff --git a/Engine/Plugins/VirtualProduction/Rivermax/RivermaxMedia/Source/RivermaxMedia/Private/RivermaxMediaPlayer.h b/Engine/Plugins/VirtualProduction/Rivermax/RivermaxMedia/Source/RivermaxMedia/Private/RivermaxMediaPlayer.h
--- a/Engine/Plugins/VirtualProduction/Rivermax/RivermaxMedia/Source/RivermaxMedia/Private/RivermaxMediaPlayer.h
+++ b/Engine/Plugins/VirtualProduction/Rivermax/RivermaxMedia/Source/RivermaxMedia/Private/RivermaxMediaPlayer.h
@@ -108,7 +108,10 @@ namespace  UE::RivermaxMedia
 		};
 
 		/** Called before sample will be converted / rendered. Used to wait for the expected frame we want to render */
-		void PreSampleUsage(FRDGBuilder& GraphBuilder, const FFrameExpectation& FrameExpectation);
+		void PreSampleUsage(const FFrameExpectation& FrameExpectation);
+		
+		/** Function waiting for the expected frame to be received. */
+		void WaitForSample(const FFrameExpectation& FrameExpectation);
 		
 		/** Called after sample was converted / rendered. Used write a fence to detect when sample is reusable */
 		void PostSampleUsage(FRDGBuilder& GraphBuilder, const FFrameExpectation& FrameExpectation);
diff --git a/Engine/Plugins/VirtualProduction/Rivermax/RivermaxMedia/Source/RivermaxMedia/Private/RivermaxMediaTextureSampleConverter.cpp b/Engine/Plugins/VirtualProduction/Rivermax/RivermaxMedia/Source/RivermaxMedia/Private/RivermaxMediaTextureSampleConverter.cpp
--- a/Engine/Plugins/VirtualProduction/Rivermax/RivermaxMedia/Source/RivermaxMedia/Private/RivermaxMediaTextureSampleConverter.cpp
+++ b/Engine/Plugins/VirtualProduction/Rivermax/RivermaxMedia/Source/RivermaxMedia/Private/RivermaxMediaTextureSampleConverter.cpp
@@ -47,7 +47,7 @@ bool FRivermaxMediaTextureSampleConverter::Convert(FTexture2DRHIRef& InDestinati
 	}
 
 	// If no input data was provided, no need to render
-	if (!Setup.GPUBuffer.IsValid() && Setup.SystemBuffer == nullptr)
+	if (Setup.GetGPUBufferFunc == nullptr && Setup.GetSystemBufferFunc == nullptr)
 	{
 		ensureMsgf(false, TEXT("Rivermax player late update succeeded but didn't provide any source data."));
 		return false;
@@ -69,14 +69,25 @@ bool FRivermaxMediaTextureSampleConverter::Convert(FTexture2DRHIRef& InDestinati
 		FRDGTextureRef OutputResource = GraphBuilder.RegisterExternalTexture(CreateRenderTarget(InDestinationTexture, TEXT("RivermaxMediaTextureOutputResource")));
 
 		// If we have a valid GPUBuffer, i.e GPUDirect is involved, use that one. Otherwise, take the system buffer and upload it in a new structured buffer.
-		if (Setup.GPUBuffer.IsValid())
+		if (Setup.GetGPUBufferFunc)
 		{
-			InputBuffer = GraphBuilder.RegisterExternalBuffer(Setup.GPUBuffer, TEXT("RMaxGPUBuffer"));
+			InputBuffer = GraphBuilder.RegisterExternalBuffer(Setup.GetGPUBufferFunc(), TEXT("RMaxGPUBuffer"));
 		}
-		else if(Setup.SystemBuffer)
+		else if (Setup.GetSystemBufferFunc)
 		{
 			TRACE_CPUPROFILER_EVENT_SCOPE(RivermaxSampleConverter::CreateStructuredBuffer);
-			InputBuffer = CreateStructuredBuffer(GraphBuilder, TEXT("RivermaxInputBuffer"), SourceBufferDesc.BytesPerElement, SourceBufferDesc.NumberOfElements, Setup.SystemBuffer, SourceBufferDesc.BytesPerElement * SourceBufferDesc.NumberOfElements, ERDGInitialDataFlags::NoCopy);
+			FRDGBufferNumElementsCallback NumElementCallback = [NumElements = SourceBufferDesc.NumberOfElements]()
+			{
+				return NumElements;
+			};
+
+			FRDGBufferInitialDataCallback GetDataCallback = Setup.GetSystemBufferFunc;
+			FRDGBufferInitialDataSizeCallback TotalSizeCallback = [TotalSize = SourceBufferDesc.BytesPerElement * SourceBufferDesc.NumberOfElements]()
+			{
+				return TotalSize;
+			};
+
+			InputBuffer = CreateStructuredBuffer(GraphBuilder, TEXT("RivermaxInputBuffer"), SourceBufferDesc.BytesPerElement, MoveTemp(NumElementCallback), MoveTemp(GetDataCallback), MoveTemp(TotalSizeCallback));
 		}
 		else
 		{
diff --git a/Engine/Plugins/VirtualProduction/Rivermax/RivermaxMedia/Source/RivermaxMedia/Private/RivermaxMediaTextureSampleConverter.h b/Engine/Plugins/VirtualProduction/Rivermax/RivermaxMedia/Source/RivermaxMedia/Private/RivermaxMediaTextureSampleConverter.h
--- a/Engine/Plugins/VirtualProduction/Rivermax/RivermaxMedia/Source/RivermaxMedia/Private/RivermaxMediaTextureSampleConverter.h
+++ b/Engine/Plugins/VirtualProduction/Rivermax/RivermaxMedia/Source/RivermaxMedia/Private/RivermaxMediaTextureSampleConverter.h
@@ -17,23 +17,24 @@ class FRivermaxMediaTextureSample;
 class FRivermaxMediaPlayer;
 
 using FPreInputConvertFunc = TUniqueFunction<void(FRDGBuilder& GraphBuilder)>;
+using FGetSystemBufferFunc = TFunction<const void *()>;
+using FGetGPUBufferFunc = TFunction<TRefCountPtr<FRDGPooledBuffer>()>;
 using FPostInputConvertFunc = TUniqueFunction<void(FRDGBuilder& GraphBuilder)>;
 
 /** Structure used during late update to let player configure some operations */
 struct FSampleConverterOperationSetup
 {
-	/** 
-	 * GPUBuffer pointer to the sample data to use during conversion 
-	 * Only valid when GPUDirect is used
-	 */
-	TRefCountPtr<FRDGPooledBuffer> GPUBuffer = nullptr;
 
-	/** System memory pointer to the sample data to use during conversion */
-	const void* SystemBuffer = nullptr;
 
 	/** Function to be called before setting up the sample conversion graph */
 	FPreInputConvertFunc PreConvertFunc = nullptr;
 
+	/** Function used to retrieve which system buffer to use. Can block until data is available. */
+	FGetSystemBufferFunc GetSystemBufferFunc = nullptr;
+	
+	/** Function used to retrieve gpu buffer if available */
+	FGetGPUBufferFunc GetGPUBufferFunc = nullptr;
+
 	/** Function to be called after setting up the sample conversion graph */
 	FPostInputConvertFunc PostConvertFunc = nullptr;
 };
