- Improving performance hit caused by buffer upload on 2110 reception

diff --git a/Engine/Plugins/VirtualProduction/Rivermax/RivermaxCore/Source/RivermaxCore/Private/RivermaxInputStream.cpp b/Engine/Plugins/VirtualProduction/Rivermax/RivermaxCore/Source/RivermaxCore/Private/RivermaxInputStream.cpp
--- a/Engine/Plugins/VirtualProduction/Rivermax/RivermaxCore/Source/RivermaxCore/Private/RivermaxInputStream.cpp
+++ b/Engine/Plugins/VirtualProduction/Rivermax/RivermaxCore/Source/RivermaxCore/Private/RivermaxInputStream.cpp
@@ -530,7 +530,7 @@ namespace UE::RivermaxCore::Private
 		}
 
 		FCUDAModule& CudaModule = FModuleManager::GetModuleChecked<FCUDAModule>("CUDA");
-
+		GPUDeviceIndex = CudaModule.GetCudaDeviceIndex();
 		CudaModule.DriverAPI()->cuCtxPushCurrent(CudaModule.GetCudaContextForDevice(GPUDeviceIndex));
 
 		// Todo: Add support for mgpu. 
diff --git a/Engine/Plugins/VirtualProduction/Rivermax/RivermaxCore/Source/RivermaxCore/Private/RivermaxManager.cpp b/Engine/Plugins/VirtualProduction/Rivermax/RivermaxCore/Source/RivermaxCore/Private/RivermaxManager.cpp
--- a/Engine/Plugins/VirtualProduction/Rivermax/RivermaxCore/Source/RivermaxCore/Private/RivermaxManager.cpp
+++ b/Engine/Plugins/VirtualProduction/Rivermax/RivermaxCore/Source/RivermaxCore/Private/RivermaxManager.cpp
@@ -281,7 +281,7 @@ namespace UE::RivermaxCore::Private
 			}
 
 			// todo: add support for mgpu compatibility verification
-			const int GPUIndex = 0;
+			const int GPUIndex = CudaModule->GetCudaDeviceIndex();
 			CUdevice CudaDevice;
 			Result = CudaModule->DriverAPI()->cuDeviceGet(&CudaDevice, GPUIndex);
 			if (Result != CUDA_SUCCESS)
diff --git a/Engine/Plugins/VirtualProduction/Rivermax/RivermaxCore/Source/RivermaxCore/Private/RivermaxOutputStream.cpp b/Engine/Plugins/VirtualProduction/Rivermax/RivermaxCore/Source/RivermaxCore/Private/RivermaxOutputStream.cpp
--- a/Engine/Plugins/VirtualProduction/Rivermax/RivermaxCore/Source/RivermaxCore/Private/RivermaxOutputStream.cpp
+++ b/Engine/Plugins/VirtualProduction/Rivermax/RivermaxCore/Source/RivermaxCore/Private/RivermaxOutputStream.cpp
@@ -1309,7 +1309,7 @@ namespace UE::RivermaxCore::Private
 		CudaModule.DriverAPI()->cuCtxPushCurrent(CudaModule.GetCudaContext());
 
 		// Todo: Add support for mgpu. For now, this will not work unless the memcpy does implicitely a cross gpu transfer.
-		constexpr int GPUIndex = 0;
+		const int GPUIndex = CudaModule.GetCudaDeviceIndex();
 		CUdevice CudaDevice;
 		CUresult Status = CudaModule.DriverAPI()->cuDeviceGet(&CudaDevice, GPUIndex);
 		if(Status != CUDA_SUCCESS)
diff --git a/Engine/Plugins/VirtualProduction/Rivermax/RivermaxMedia/Source/RivermaxMedia/Private/RivermaxMediaPlayer.cpp b/Engine/Plugins/VirtualProduction/Rivermax/RivermaxMedia/Source/RivermaxMedia/Private/RivermaxMediaPlayer.cpp
--- a/Engine/Plugins/VirtualProduction/Rivermax/RivermaxMedia/Source/RivermaxMedia/Private/RivermaxMediaPlayer.cpp
+++ b/Engine/Plugins/VirtualProduction/Rivermax/RivermaxMedia/Source/RivermaxMedia/Private/RivermaxMediaPlayer.cpp
@@ -42,6 +42,13 @@ namespace UE::RivermaxMedia
 		TEXT("Override latency in framelock mode. 0 for 0 frame of latency and 1 for 1 frame of latency."),
 		ECVF_Default);
 
+	static TAutoConsoleVariable<int32> CVarRivermaxSampleUploadMode(
+		TEXT("Rivermax.Player.UploadMode"),
+		1,
+		TEXT("Mode 0: Upload is done on the render thread.\n"
+			"Mode 1: Upload is done in its own thread before being rendered."),
+		ECVF_Default);
+
 	/* FRivermaxVideoPlayer structors
 	 *****************************************************************************/
 
@@ -286,6 +293,7 @@ namespace UE::RivermaxMedia
 				if(RivermaxThreadCurrentTextureSample->ReceptionState == ESampleReceptionState::Receiving)
 				{
 					RivermaxThreadCurrentTextureSample->ReceptionState = ESampleReceptionState::Received;
+					RivermaxThreadCurrentTextureSample->bIsReadyToRender = bDoesStreamSupportsGPUDirect;
 				}
 				else
 				{
@@ -548,18 +556,15 @@ namespace UE::RivermaxMedia
 			return false;
 		}
 
-		// If it's the first time we are rendering, make previous frame number as available since we will never get to them
-		if (FrameTracking.LastFrameRendered.IsSet() && NextFrameExpectations.FrameNumber > (FrameTracking.LastFrameRendered.GetValue() + 1))
+		// Always clean up frames from the past. They will never be used
+		for (const TSharedPtr<FRivermaxSampleWrapper>& Frame : SamplePool)
 		{
-			for (const TSharedPtr<FRivermaxSampleWrapper>& Frame : SamplePool)
+			if (Frame->ReceptionState != ESampleReceptionState::Available)
 			{
-				if (Frame->ReceptionState != ESampleReceptionState::Available)
+				if (Frame->FrameNumber < NextFrameExpectations.FrameNumber)
 				{
-					if (Frame->FrameNumber < NextFrameExpectations.FrameNumber)
-					{
-						Frame->ReceptionState = ESampleReceptionState::Available;
-						UE_LOG(LogRivermaxMedia, Verbose, TEXT("Making frame %u as available since it will never be processed. Last render = %u, next render = %u"), Frame->FrameNumber, FrameTracking.LastFrameRendered.GetValue(), NextFrameExpectations.FrameNumber);
-					}
+					Frame->ReceptionState = ESampleReceptionState::Available;
+					UE_LOG(LogRivermaxMedia, Verbose, TEXT("Making frame %u as available since it will never be processed. Last render = %u, next render = %u"), Frame->FrameNumber, FrameTracking.LastFrameRendered.GetValue(), NextFrameExpectations.FrameNumber);
 				}
 			}
 		}
@@ -591,30 +596,18 @@ namespace UE::RivermaxMedia
 		SamplePool[NextFrameExpectations.FrameIndex]->bIsPendingRendering = true;
 
 		{
-			// Provide the late sample source data for the converter
-			const TSharedPtr<FRivermaxMediaTextureSample>& NextSample = SamplePool[NextFrameExpectations.FrameIndex]->Sample;
-			if (bDoesStreamSupportsGPUDirect)
+			const int32 SampleUploadMode = CVarRivermaxSampleUploadMode.GetValueOnRenderThread();
+			if (SampleUploadMode == 0)
 			{
-				// For gpudirect path, we return the expected buffer directly
-				// PreConvertFunc will take care of waiting on RHI thread for the buffer to be usable
-				OutConverterSetup.GetGPUBufferFunc = [Sample = NextSample](){ return Sample->GetGPUBuffer(); };
-				OutConverterSetup.PreConvertFunc = [NextFrameExpectations, this](const FRDGBuilder& GraphBuilder)
-				{
-					PreSampleUsage(NextFrameExpectations);
-				};
+				SampleUploadSetupRenderThreadMode(NextFrameExpectations, OutConverterSetup);
+				
 			}
-			else
+			else 
 			{
-				// For system memory path, we wait for the sample to be received before returning it.
-				// This will stall the render thread in its current state
-				OutConverterSetup.GetSystemBufferFunc = [NextFrameExpectations, Sample = NextSample, this]()
-				{
-					WaitForSample(NextFrameExpectations);
-					return Sample->GetBuffer();
-				};
+				SampleUploadSetupTaskThreadMode(NextFrameExpectations, OutConverterSetup);
 			}
 
-			// Setup post sample usage pass
+			// Setup post sample usage pass 
 			OutConverterSetup.PostConvertFunc = [NextFrameExpectations, this] (FRDGBuilder& GraphBuilder)
 			{
 				PostSampleUsage(GraphBuilder, NextFrameExpectations);
@@ -793,6 +786,12 @@ namespace UE::RivermaxMedia
 					if ((FPlatformTime::Seconds() - StartTimeSeconds) > TimeoutSeconds)
 					{
 						UE_LOG(LogRivermaxMedia, Error, TEXT("Timed out waiting for frame %u to be rendered to receive frame %u into."), ExpectedFrame->FrameNumber, FrameInfo.FrameNumber);
+					
+						// Mark this frame as skipped since we're stealing it
+						FScopeLock Lock(&SkippedFrameCriticalSection);
+						const TInterval<uint32> Interval(ExpectedFrame->FrameNumber, ExpectedFrame->FrameNumber);
+						SkippedFrames.Add(Interval);
+
 						break;
 					}
 				}
@@ -845,18 +844,6 @@ namespace UE::RivermaxMedia
 		return IMediaPlayer::GetPlayerFeatureFlag(flag);
 	}
 
-	void FRivermaxMediaPlayer::PreSampleUsage(const FFrameExpectation& FrameExpectation)
-	{
-		// Before the sample is used during conversion, we want to make sure the expected frame is available to be used
-		FRHICommandListImmediate& RHICmdList = FRHICommandListExecutor::GetImmediateCommandList();
-		RHICmdList.EnqueueLambda(
-			[FrameExpectation, this](FRHICommandList& RHICmdList)
-			{
-				TRACE_CPUPROFILER_EVENT_SCOPE(RivermaxWaitForSamples);
-				WaitForSample(FrameExpectation);
-			});
-	}
-
 	void FRivermaxMediaPlayer::PostSampleUsage(FRDGBuilder& GraphBuilder, const FFrameExpectation& FrameExpectation)
 	{
 		GraphBuilder.AddPass(RDG_EVENT_NAME("RivermaxPostSampleUsage"),
@@ -879,7 +866,7 @@ namespace UE::RivermaxMedia
 							--TasksInFlight;
 						};
 
-						TRACE_CPUPROFILER_EVENT_SCOPE(RivermaxSampleConversionDone);
+						TRACE_CPUPROFILER_EVENT_SCOPE(RmaxWaitForShader);
 						do
 						{
 							const bool bHasValidFence = SamplePool[FrameExpectation.FrameIndex]->SampleConversionFence.IsValid();
@@ -899,7 +886,9 @@ namespace UE::RivermaxMedia
 						// We clear the fence and signal that it can be re-used.
 						SamplePool[FrameExpectation.FrameIndex]->SampleConversionFence->Clear();
 						SamplePool[FrameExpectation.FrameIndex]->ReceptionState = ESampleReceptionState::Available;
+						SamplePool[FrameExpectation.FrameIndex]->bIsReadyToRender = false;
 						SamplePool[FrameExpectation.FrameIndex]->bIsPendingRendering = false;
+						SamplePool[FrameExpectation.FrameIndex]->LockedMemory = nullptr;
 						TryClearSkippedInterval(FrameExpectation.FrameNumber);
 					});
 			});
@@ -1003,10 +992,8 @@ namespace UE::RivermaxMedia
 		}
 	}
 
-	void FRivermaxMediaPlayer::WaitForSample(const FFrameExpectation& FrameExpectation)
+	void FRivermaxMediaPlayer::WaitForSample(const FFrameExpectation& FrameExpectation, FWaitConditionFunc WaitConditionFunction, bool bCanTimeout)
 	{
-		TRACE_CPUPROFILER_EVENT_SCOPE(RivermaxWaitForSample);
-
 		const double StartTimeSeconds = FPlatformTime::Seconds();
 		constexpr double TimeoutSeconds = 0.5;
 
@@ -1031,28 +1018,147 @@ namespace UE::RivermaxMedia
 
 				// Our goal here is to wait until the expected frame is available to be used (received) unless there is a timeout
 				const TSharedPtr<FRivermaxSampleWrapper>& Frame = SamplePool[FrameExpectation.FrameIndex];
-				if (Frame->ReceptionState == ESampleReceptionState::Received)
+				if (WaitConditionFunction(Frame))
 				{
-					if (Frame->FrameNumber >= FrameExpectation.FrameNumber)
+					if (Frame->FrameNumber != FrameExpectation.FrameNumber)
 					{
-						if (Frame->FrameNumber > FrameExpectation.FrameNumber)
-						{
-							UE_LOG(LogRivermaxMedia, Warning, TEXT("More recent frame, %u, than expected, %u, found while waiting."), Frame->FrameNumber, FrameExpectation.FrameNumber);
-						}
+						UE_LOG(LogRivermaxMedia, Warning, TEXT("Rendering unexpected frame %u, when frame %u was expected."), Frame->FrameNumber, FrameExpectation.FrameNumber);
+					}
+					break;
+				}
+
+				{
+					FPlatformProcess::SleepNoStats(SleepTimeSeconds);
+					if (bCanTimeout && ((FPlatformTime::Seconds() - StartTimeSeconds) > TimeoutSeconds))
+					{
+						UE_LOG(LogRivermaxMedia, Error, TEXT("Timed out waiting for frame %u."), FrameExpectation.FrameNumber);
 						break;
 					}
 				}
 			}
+		}
+	}
+
+	void FRivermaxMediaPlayer::SampleUploadSetupRenderThreadMode(const FFrameExpectation& NextFrameExpectations, FSampleConverterOperationSetup& OutConverterSetup)
+	{
+		TSharedPtr<FRivermaxSampleWrapper> SampleWrapper = SamplePool[NextFrameExpectations.FrameIndex];
+		if (bDoesStreamSupportsGPUDirect)
+		{
+			OutConverterSetup.GetGPUBufferFunc = [SampleWrapper]() { return SampleWrapper->Sample->GetGPUBuffer(); };
+
+			// Setup requirements for sample to be ready to be rendered
+			FRHICommandListImmediate& RHICmdList = FRHICommandListExecutor::GetImmediateCommandList();
+			RHICmdList.EnqueueLambda(
+				[NextFrameExpectations, this](FRHICommandList& RHICmdList)
+				{
+					FWaitConditionFunc ReadyToRenderConditionFunc = [](const TSharedPtr<FRivermaxSampleWrapper>& Sample)
+					{
+						return Sample->bIsReadyToRender.load();
+					};
+
+					TRACE_CPUPROFILER_EVENT_SCOPE(RmaxWaitSampleReadyness);
+					constexpr bool bCanTimeout = true;
+					WaitForSample(NextFrameExpectations, MoveTemp(ReadyToRenderConditionFunc), bCanTimeout);
+				}
+			);
+		}
+		else
+		{
+			// For system memory path, we wait for the sample to be received before returning it.
+			OutConverterSetup.GetSystemBufferFunc = [NextFrameExpectations, SampleWrapper, this]()
+			{
+				FWaitConditionFunc ReceivedConditionFunc = [](const TSharedPtr<FRivermaxSampleWrapper>& Sample)
+				{
+					return Sample->ReceptionState == ESampleReceptionState::Received;
+				};
+
+				TRACE_CPUPROFILER_EVENT_SCOPE(RmaxWaitSampleReception);
+				constexpr bool bCanTimeout = true;
+				WaitForSample(NextFrameExpectations, MoveTemp(ReceivedConditionFunc), bCanTimeout);
+
+				// We can return the system buffer once we know it has arrived
+				return SampleWrapper->Sample->GetBuffer();
+			};
+		}
+	}
+
+	void FRivermaxMediaPlayer::SampleUploadSetupTaskThreadMode(const FFrameExpectation& NextFrameExpectations, FSampleConverterOperationSetup& OutConverterSetup)
+	{
+		TSharedPtr<FRivermaxSampleWrapper> SampleWrapper = SamplePool[NextFrameExpectations.FrameIndex];
+
+		// We will always be providing a buffer already located on the GPU even when not using gpudirect
+		// Once a frame has arrived on system, we will upload it to the allocated gpu buffer.
+		OutConverterSetup.GetGPUBufferFunc = [SampleWrapper]() { return SampleWrapper->Sample->GetGPUBuffer(); };
 
+		OutConverterSetup.PreConvertFunc = [NextFrameExpectations, SampleWrapper, this](const FRDGBuilder& GraphBuilder)
+		{
+			// When GPUDirect is not involved, we have an extra step to do. We need to wait for the sample to be received
+			// but also initiate the memcopy to gpu memory for it to be rendered
+			if (bDoesStreamSupportsGPUDirect == false)
 			{
-				FPlatformProcess::SleepNoStats(SleepTimeSeconds);
-				if ((FPlatformTime::Seconds() - StartTimeSeconds) > TimeoutSeconds)
+				constexpr uint32 Offset = 0;
+				const uint32 Size = SampleWrapper->Sample->GetGPUBuffer()->GetSize();
+				SampleWrapper->LockedMemory = GraphBuilder.RHICmdList.LockBuffer(SampleWrapper->Sample->GetGPUBuffer()->GetRHI(), Offset, Size, EResourceLockMode::RLM_WriteOnly);
+
+				++TasksInFlight;
+				UE::Tasks::Launch(UE_SOURCE_LOCATION,
+					[NextFrameExpectations, Size, SampleWrapper, this]()
+					{
+						ON_SCOPE_EXIT
+						{
+							--TasksInFlight;
+						};
+
+						FWaitConditionFunc ReceivedConditionFunc = [](const TSharedPtr<FRivermaxSampleWrapper>& Sample)
+						{
+							return Sample->ReceptionState == ESampleReceptionState::Received;
+						};
+
+						{
+							TRACE_CPUPROFILER_EVENT_SCOPE(RmaxWaitSampleReception);
+							constexpr bool bCanTimeout = true;
+							WaitForSample(NextFrameExpectations, MoveTemp(ReceivedConditionFunc), bCanTimeout);
+						}
+
+						// In case the wait failed, make sure it's received before copying
+						if (SamplePool[NextFrameExpectations.FrameIndex]->ReceptionState == ESampleReceptionState::Received)
+						{
+							TRACE_CPUPROFILER_EVENT_SCOPE(RmaxSampleUpload);
+							FMemory::Memcpy(SampleWrapper->LockedMemory, SampleWrapper->Sample->GetBuffer(), Size);
+						}
+						
+						// We always consider the sample ready to be rendered even if we haven't copied something over (timeout)
+						// RHI commands have already been submitted to look for that frame so we can't back out at this point
+						SampleWrapper->bIsReadyToRender = true;
+					}
+				);
+			}
+
+			// Setup requirements for sample to be ready to be rendered
+			FRHICommandListImmediate & RHICmdList = FRHICommandListExecutor::GetImmediateCommandList();
+			RHICmdList.EnqueueLambda(
+				[NextFrameExpectations, this](FRHICommandList& RHICmdList)
 				{
-					UE_LOG(LogRivermaxMedia, Error, TEXT("Timed out waiting for frame %u."), FrameExpectation.FrameNumber);
-					break;
+					FWaitConditionFunc ReadyToRenderConditionFunc = [](const TSharedPtr<FRivermaxSampleWrapper>& Sample)
+					{
+						return Sample->bIsReadyToRender.load();
+					};
+
+					TRACE_CPUPROFILER_EVENT_SCOPE(RmaxWaitSampleReadyness);
+
+					// Only accept timeouts here if we're going through the system memory in order to avoid timing out while doing
+					// the memcopy (buffer upload)
+					bool bCanTimeout = bDoesStreamSupportsGPUDirect;
+					WaitForSample(NextFrameExpectations, MoveTemp(ReadyToRenderConditionFunc), bCanTimeout);
 				}
+			);
+
+			// Final step, if the memory was locked (non gpu direct), enqueue unlock after the wait for sample in order to render it
+			if (SampleWrapper->LockedMemory && ensure(!bDoesStreamSupportsGPUDirect))
+			{
+				RHICmdList.UnlockBuffer(SampleWrapper->Sample->GetGPUBuffer()->GetRHI());
 			}
-		}
+		};
 	}
 
 }
diff --git a/Engine/Plugins/VirtualProduction/Rivermax/RivermaxMedia/Source/RivermaxMedia/Private/RivermaxMediaPlayer.h b/Engine/Plugins/VirtualProduction/Rivermax/RivermaxMedia/Source/RivermaxMedia/Private/RivermaxMediaPlayer.h
--- a/Engine/Plugins/VirtualProduction/Rivermax/RivermaxMedia/Source/RivermaxMedia/Private/RivermaxMediaPlayer.h
+++ b/Engine/Plugins/VirtualProduction/Rivermax/RivermaxMedia/Source/RivermaxMedia/Private/RivermaxMediaPlayer.h
@@ -7,6 +7,7 @@
 #include "RivermaxMediaSource.h"
 #include "RivermaxMediaTextureSampleConverter.h"
 #include "RivermaxTypes.h"
+#include "Templates/Function.h"
 
 #include <atomic>
 
@@ -107,11 +108,55 @@ namespace  UE::RivermaxMedia
 			uint32 FrameIndex = 0;
 		};
 
-		/** Called before sample will be converted / rendered. Used to wait for the expected frame we want to render */
-		void PreSampleUsage(const FFrameExpectation& FrameExpectation);
+		enum class ESampleReceptionState : uint8
+		{
+			// Sample is ready to be requested by rivermax stream
+			Available,
+
+			// Sample has been requested and is being received
+			Receiving,
+
+			// Sample has been received and is ready to be rendered
+			Received,
+		};
+
+		/** Wrapper around a sample holding more information about its state */
+		struct FRivermaxSampleWrapper
+		{
+			/** State of this sample */
+			std::atomic<ESampleReceptionState> ReceptionState = ESampleReceptionState::Available;
+
+			/** True when queued for rendering. Will be false once fence has been written, after shader usage. */
+			std::atomic<bool> bIsPendingRendering = false;
+
+			/** True if sample can be rendered. For non-gpudirect, we need to copy sample from system memory to gpu to be ready */
+			std::atomic<bool> bIsReadyToRender = false;
+
+			/** Locked memory of gpu buffer when uploading */
+			void* LockedMemory = nullptr;
+
+			/** Actual sample buffer container used by media framework */
+			TSharedPtr<FRivermaxMediaTextureSample> Sample;
+
+			/** Write fence enqueued after sample conversion to know when it's ready to be reused */
+			FGPUFenceRHIRef SampleConversionFence;
+
+			/** Frame number of this sample */
+			uint32 FrameNumber = 0;
+
+			/** Timestamp of this sample */
+			uint32 Timestamp = 0;
+		};
+
+		/** Buffer upload setup that will block render thread while waiting for sample and uploading it */
+		void SampleUploadSetupRenderThreadMode(const FFrameExpectation& FrameExpectation, FSampleConverterOperationSetup& OutConverterSetup);
 		
+		/** Buffer upload setup that will wait on its own task to wait for sample and do the upload */
+		void SampleUploadSetupTaskThreadMode(const FFrameExpectation& FrameExpectation, FSampleConverterOperationSetup& OutConverterSetup);
+
 		/** Function waiting for the expected frame to be received. */
-		void WaitForSample(const FFrameExpectation& FrameExpectation);
+		using FWaitConditionFunc = TUniqueFunction<bool(const TSharedPtr<FRivermaxSampleWrapper>&)>;
+		void WaitForSample(const FFrameExpectation& FrameExpectation, FWaitConditionFunc WaitConditionFunction, bool bCanTimeout);
 		
 		/** Called after sample was converted / rendered. Used write a fence to detect when sample is reusable */
 		void PostSampleUsage(FRDGBuilder& GraphBuilder, const FFrameExpectation& FrameExpectation);
@@ -141,30 +186,6 @@ namespace  UE::RivermaxMedia
 
 	private:
 
-		enum class ESampleReceptionState : uint8
-		{
-			// Sample is ready to be requested by rivermax stream
-			Available,
-
-			// Sample has been requested and is being received
-			Receiving,
-
-			// Sample has been received and is ready to be rendered
-			Received,
-		};
-
-		/** Wrapper around a sample holding more information about its state */
-		struct FRivermaxSampleWrapper
-		{
-			std::atomic<ESampleReceptionState> ReceptionState = ESampleReceptionState::Available;
-			bool bIsPendingRendering = false;
-
-			TSharedPtr<FRivermaxMediaTextureSample> Sample;
-			FGPUFenceRHIRef SampleConversionFence;
-			uint32 FrameNumber = 0;
-			uint32 Timestamp = 0;
-		};
-
 		TSharedPtr<FRivermaxSampleWrapper> RivermaxThreadCurrentTextureSample;
 
 		/** Common gpu buffer to use when we render a sample */
diff --git a/Engine/Source/Runtime/CUDA/Source/Private/CudaModule.cpp b/Engine/Source/Runtime/CUDA/Source/Private/CudaModule.cpp
--- a/Engine/Source/Runtime/CUDA/Source/Private/CudaModule.cpp
+++ b/Engine/Source/Runtime/CUDA/Source/Private/CudaModule.cpp
@@ -360,6 +360,11 @@ CUcontext FCUDAModule::GetCudaContext()
 	return contextMap[rhiDeviceIndex];
 }
 
+uint32 FCUDAModule::GetCudaDeviceIndex() const
+{
+	return rhiDeviceIndex;
+}
+
 CUcontext FCUDAModule::GetCudaContextForDevice(int DeviceIndex)
 {
 	if (IsAvailable() == false)
diff --git a/Engine/Source/Runtime/CUDA/Source/Public/CudaModule.h b/Engine/Source/Runtime/CUDA/Source/Public/CudaModule.h
--- a/Engine/Source/Runtime/CUDA/Source/Public/CudaModule.h
+++ b/Engine/Source/Runtime/CUDA/Source/Public/CudaModule.h
@@ -32,7 +32,10 @@ public:
 
 	// Retrieves the CUDA context for the GPU device currently in use by the Vulkan RHI
 	CUcontext GetCudaContext();
-
+	
+	// Retrieves the device index for the current context
+	uint32 GetCudaDeviceIndex() const;
+	
 	// Retrieves or creates the CUDA context for the specified GPU device
 	CUcontext GetCudaContextForDevice(int DeviceIndex);
 
