Fixed typo in Rivrmax player

diff --git a/Engine/Plugins/VirtualProduction/Rivermax/RivermaxMedia/Source/RivermaxMedia/Private/RivermaxMediaPlayer.cpp b/Engine/Plugins/VirtualProduction/Rivermax/RivermaxMedia/Source/RivermaxMedia/Private/RivermaxMediaPlayer.cpp
--- a/Engine/Plugins/VirtualProduction/Rivermax/RivermaxMedia/Source/RivermaxMedia/Private/RivermaxMediaPlayer.cpp
+++ b/Engine/Plugins/VirtualProduction/Rivermax/RivermaxMedia/Source/RivermaxMedia/Private/RivermaxMediaPlayer.cpp
@@ -659,7 +659,7 @@ namespace UE::RivermaxMedia
 			int32 HighestTimestampIndex = INDEX_NONE;
 			uint32 HighestTimestamp = 0;
 
-			int32 LowestTimstampIndex = INDEX_NONE;
+			int32 LowestTimestampIndex = INDEX_NONE;
 			uint32 LowestTimestamp = ~0;
 
 			// Use timestamps to pick the latest one. 
@@ -678,7 +678,7 @@ namespace UE::RivermaxMedia
 					if (LowestTimestamp > Frame->Timestamp)
 					{
 						LowestTimestamp = Frame->Timestamp;
-						LowestTimstampIndex = Index;
+						LowestTimestampIndex = Index;
 					}
 				}
 			}
@@ -689,7 +689,7 @@ namespace UE::RivermaxMedia
 				return false;
 			}
 
-			check(LowestTimstampIndex >= 0);
+			check(LowestTimestampIndex >= 0);
 
 			// No point in rendering the same frame as last time
 			if (FrameTracking.LastFrameExpectation.FrameNumber == SamplePool[HighestTimestampIndex]->FrameNumber)
