- Fixing crash in rivermax verbose log

diff --git a/Engine/Plugins/VirtualProduction/Rivermax/RivermaxMedia/Source/RivermaxMedia/Private/RivermaxMediaPlayer.cpp b/Engine/Plugins/VirtualProduction/Rivermax/RivermaxMedia/Source/RivermaxMedia/Private/RivermaxMediaPlayer.cpp
--- a/Engine/Plugins/VirtualProduction/Rivermax/RivermaxMedia/Source/RivermaxMedia/Private/RivermaxMediaPlayer.cpp
+++ b/Engine/Plugins/VirtualProduction/Rivermax/RivermaxMedia/Source/RivermaxMedia/Private/RivermaxMediaPlayer.cpp
@@ -565,7 +565,9 @@ namespace UE::RivermaxMedia
 				if (Frame->FrameNumber < NextFrameExpectations.FrameNumber)
 				{
 					Frame->ReceptionState = ESampleReceptionState::Available;
-					UE_LOG(LogRivermaxMedia, Verbose, TEXT("Making frame %u as available since it will never be processed. Last render = %u, next render = %u"), Frame->FrameNumber, FrameTracking.LastFrameRendered.GetValue(), NextFrameExpectations.FrameNumber);
+
+					const FString LastRender = FrameTracking.LastFrameRendered.IsSet() ? FString::Printf(TEXT("%u"), FrameTracking.LastFrameRendered.GetValue()) : FString(TEXT("None"));
+					UE_LOG(LogRivermaxMedia, Verbose, TEXT("Making frame %u as available since it will never be processed. Last render = %s, next render = %u"), Frame->FrameNumber, *LastRender, NextFrameExpectations.FrameNumber);
 				}
 			}
 		}
