- Enabling block on reservation for frame creation alignment used for inner streaming with 2110

diff --git a/Engine/Plugins/VirtualProduction/Rivermax/RivermaxMedia/Source/RivermaxMedia/Public/RivermaxMediaOutput.h b/Engine/Plugins/VirtualProduction/Rivermax/RivermaxMedia/Source/RivermaxMedia/Public/RivermaxMediaOutput.h
--- a/Engine/Plugins/VirtualProduction/Rivermax/RivermaxMedia/Source/RivermaxMedia/Public/RivermaxMediaOutput.h
+++ b/Engine/Plugins/VirtualProduction/Rivermax/RivermaxMedia/Source/RivermaxMedia/Public/RivermaxMediaOutput.h
@@ -91,7 +91,7 @@ public:
 	bool bDoContinuousOutput = true;
 
 	/** For alignment point mode, controls whether we stall engine before capturing if there are no buffer available to capture into */
-	UPROPERTY(EditAnywhere, BlueprintReadWrite, Category = "Output", meta = (EditCondition = "AlignmentMode != ERivermaxMediaAlignmentMode::FrameCreation"))
+	UPROPERTY(EditAnywhere, BlueprintReadWrite, Category = "Output")
 	ERivermaxFrameLockingMode FrameLockingMode = ERivermaxFrameLockingMode::FreeRun;
 
 	/** Number of frames that can be queued / used in output queue. Frame being sent counts for 1. */
