Fix rivermax media player flickering

diff --git a/Engine/Plugins/VirtualProduction/Rivermax/RivermaxMedia/Source/RivermaxMedia/Private/RivermaxMediaTextureSample.cpp b/Engine/Plugins/VirtualProduction/Rivermax/RivermaxMedia/Source/RivermaxMedia/Private/RivermaxMediaTextureSample.cpp
--- a/Engine/Plugins/VirtualProduction/Rivermax/RivermaxMedia/Source/RivermaxMedia/Private/RivermaxMediaTextureSample.cpp
+++ b/Engine/Plugins/VirtualProduction/Rivermax/RivermaxMedia/Source/RivermaxMedia/Private/RivermaxMediaTextureSample.cpp
@@ -35,21 +35,20 @@ bool FRivermaxMediaTextureSample::IsOutputSrgb() const
 
 bool FRivermaxMediaTextureSample::ConfigureSample(const FSampleConfigurationArgs& Args)
 {
-	EMediaTextureSampleFormat VideoSampleFormat;
 	switch (Args.SampleFormat)
 	{
 		case ERivermaxMediaSourcePixelFormat::RGB_12bit:
 			// Falls through
 		case ERivermaxMediaSourcePixelFormat::RGB_16bit_Float:
 		{
-			VideoSampleFormat = EMediaTextureSampleFormat::FloatRGBA;
+			SampleFormat = EMediaTextureSampleFormat::FloatRGBA;
 			break;
 		}
 		case ERivermaxMediaSourcePixelFormat::RGB_10bit:
 			// Falls through
 		case ERivermaxMediaSourcePixelFormat::YUV422_10bit:
 		{
-			VideoSampleFormat = EMediaTextureSampleFormat::CharBGR10A2;
+			SampleFormat = EMediaTextureSampleFormat::CharBGR10A2;
 			break;
 		}
 		case ERivermaxMediaSourcePixelFormat::YUV422_8bit:
@@ -58,7 +57,7 @@ bool FRivermaxMediaTextureSample::ConfigureSample(const FSampleConfigurationArgs
 			// Falls through
 		default:
 		{
-			VideoSampleFormat = EMediaTextureSampleFormat::CharBGRA;
+			SampleFormat = EMediaTextureSampleFormat::CharBGRA;
 			break;
 		}
 	}
diff --git a/Engine/Plugins/VirtualProduction/Rivermax/RivermaxMedia/Source/RivermaxMedia/Private/RivermaxMediaTextureSample.h b/Engine/Plugins/VirtualProduction/Rivermax/RivermaxMedia/Source/RivermaxMedia/Private/RivermaxMediaTextureSample.h
--- a/Engine/Plugins/VirtualProduction/Rivermax/RivermaxMedia/Source/RivermaxMedia/Private/RivermaxMediaTextureSample.h
+++ b/Engine/Plugins/VirtualProduction/Rivermax/RivermaxMedia/Source/RivermaxMedia/Private/RivermaxMediaTextureSample.h
@@ -113,7 +113,7 @@ private:
 	FIntPoint Dimension = FIntPoint::ZeroValue;
 
 	/** Sample format. */
-	EMediaTextureSampleFormat SampleFormat;
+	EMediaTextureSampleFormat SampleFormat = EMediaTextureSampleFormat::CharBGRA;
 
 	/** Format in the rivermax realm */
 	ERivermaxMediaSourcePixelFormat InputFormat;
