Fixing crash when an error occurs when rivermax players open

diff --git a/Engine/Plugins/VirtualProduction/Rivermax/RivermaxMedia/Source/RivermaxMedia/Private/RivermaxMediaPlayer.cpp b/Engine/Plugins/VirtualProduction/Rivermax/RivermaxMedia/Source/RivermaxMedia/Private/RivermaxMediaPlayer.cpp
--- a/Engine/Plugins/VirtualProduction/Rivermax/RivermaxMedia/Source/RivermaxMedia/Private/RivermaxMediaPlayer.cpp
+++ b/Engine/Plugins/VirtualProduction/Rivermax/RivermaxMedia/Source/RivermaxMedia/Private/RivermaxMediaPlayer.cpp
@@ -84,6 +84,7 @@ namespace UE::RivermaxMedia
 	{
 		if (!Super::Open(Url, Options))
 		{
+			RivermaxThreadNewState = EMediaState::Error;
 			return false;
 		}
 
@@ -123,7 +124,6 @@ namespace UE::RivermaxMedia
 		if (InputStream == nullptr || !InputStream->Initialize(StreamOptions, *this))
 		{
 			UE_LOG(LogRivermaxMedia, Warning, TEXT("Failed to initialize Rivermax input stream."));
-			CurrentState = EMediaState::Error;
 			RivermaxThreadNewState = EMediaState::Error;
 			InputStream.Reset();
 			return false;
