- Adding support for timecode provider based on rivermax time

diff --git a/Engine/Plugins/VirtualProduction/Rivermax/RivermaxMedia/Source/RivermaxMedia/Private/RivermaxMediaCapture.cpp b/Engine/Plugins/VirtualProduction/Rivermax/RivermaxMedia/Source/RivermaxMedia/Private/RivermaxMediaCapture.cpp
--- a/Engine/Plugins/VirtualProduction/Rivermax/RivermaxMedia/Source/RivermaxMedia/Private/RivermaxMediaCapture.cpp
+++ b/Engine/Plugins/VirtualProduction/Rivermax/RivermaxMedia/Source/RivermaxMedia/Private/RivermaxMediaCapture.cpp
@@ -280,12 +280,11 @@ void URivermaxMediaCapture::AddFrameReservationPass(FRDGBuilder& GraphBuilder)
 	});
 }
 
-void URivermaxMediaCapture::OnFrameCaptured_RenderingThread(const FCaptureBaseData& InBaseData, TSharedPtr<FMediaCaptureUserData, ESPMode::ThreadSafe> InUserData, void* InBuffer, int32 Width, int32 Height, int32 BytesPerRow)
+void URivermaxMediaCapture::OnFrameCapturedInternal_AnyThread(const FCaptureBaseData& InBaseData, TSharedPtr<FMediaCaptureUserData, ESPMode::ThreadSafe> InUserData, void* InBuffer, int32 Width, int32 Height, int32 BytesPerRow)
 {
 	using namespace UE::RivermaxCore;
+	TRACE_CPUPROFILER_EVENT_SCOPE_TEXT(*FRivermaxTracingUtils::RmaxOutMediaCapturePipeTraceEvents[InBaseData.SourceFrameNumberRenderThread % 10]);
 
-	TRACE_CPUPROFILER_EVENT_SCOPE(URivermaxMediaCapture::OnFrameCaptured_RenderingThread);
-	
 	FRivermaxOutputVideoFrameInfo NewFrame;
 	NewFrame.Height = Height;
 	NewFrame.Width = Width;
@@ -298,11 +297,9 @@ void URivermaxMediaCapture::OnFrameCaptured_RenderingThread(const FCaptureBaseDa
 	}
 }
 
-void URivermaxMediaCapture::OnRHIResourceCaptured_AnyThread(const FCaptureBaseData& InBaseData, TSharedPtr<FMediaCaptureUserData, ESPMode::ThreadSafe> InUserData, FBufferRHIRef InBuffer)
+void URivermaxMediaCapture::OnRHIResourceCapturedInternal_AnyThread(const FCaptureBaseData& InBaseData, TSharedPtr<FMediaCaptureUserData, ESPMode::ThreadSafe> InUserData, FBufferRHIRef InBuffer)
 {
 	using namespace UE::RivermaxCore;
-
-	TRACE_CPUPROFILER_EVENT_SCOPE(URivermaxMediaCapture::OnRHIResourceCaptured_AnyThread);
 	TRACE_CPUPROFILER_EVENT_SCOPE_TEXT(*FRivermaxTracingUtils::RmaxOutMediaCapturePipeTraceEvents[InBaseData.SourceFrameNumberRenderThread % 10]);
 
 	FRivermaxOutputVideoFrameInfo NewFrame;
@@ -314,23 +311,32 @@ void URivermaxMediaCapture::OnRHIResourceCaptured_AnyThread(const FCaptureBaseDa
 	}
 }
 
+void URivermaxMediaCapture::OnFrameCaptured_RenderingThread(const FCaptureBaseData& InBaseData, TSharedPtr<FMediaCaptureUserData, ESPMode::ThreadSafe> InUserData, void* InBuffer, int32 Width, int32 Height, int32 BytesPerRow)
+{
+	TRACE_CPUPROFILER_EVENT_SCOPE(URivermaxMediaCapture::OnFrameCaptured_RenderingThread);
+	OnFrameCapturedInternal_AnyThread(InBaseData, InUserData, InBuffer, Width, Height, BytesPerRow);
+}
+
+void URivermaxMediaCapture::OnRHIResourceCaptured_RenderingThread(const FCaptureBaseData& InBaseData, TSharedPtr<FMediaCaptureUserData, ESPMode::ThreadSafe> InUserData, FBufferRHIRef InBuffer)
+{
+	TRACE_CPUPROFILER_EVENT_SCOPE(URivermaxMediaCapture::OnRHIResourceCaptured_RenderingThread);
+	OnRHIResourceCapturedInternal_AnyThread(InBaseData, InUserData, InBuffer);
+}
+
+void URivermaxMediaCapture::OnRHIResourceCaptured_AnyThread(const FCaptureBaseData& InBaseData, TSharedPtr<FMediaCaptureUserData, ESPMode::ThreadSafe> InUserData, FBufferRHIRef InBuffer)
+{
+	using namespace UE::RivermaxCore;
+
+	TRACE_CPUPROFILER_EVENT_SCOPE(URivermaxMediaCapture::OnRHIResourceCaptured_AnyThread);
+	OnRHIResourceCapturedInternal_AnyThread(InBaseData, InUserData, InBuffer);
+}
+
 void URivermaxMediaCapture::OnFrameCaptured_AnyThread(const FCaptureBaseData& InBaseData, TSharedPtr<FMediaCaptureUserData, ESPMode::ThreadSafe> InUserData, const FMediaCaptureResourceData& InResourceData)
 {
 	using namespace UE::RivermaxCore;
 
 	TRACE_CPUPROFILER_EVENT_SCOPE(URivermaxMediaCapture::OnFrameCaptured_AnyThread);
-	TRACE_CPUPROFILER_EVENT_SCOPE_TEXT(*FRivermaxTracingUtils::RmaxOutMediaCapturePipeTraceEvents[InBaseData.SourceFrameNumberRenderThread % 10]);
-	
-	FRivermaxOutputVideoFrameInfo NewFrame;
-	NewFrame.Height = InResourceData.Height;
-	NewFrame.Width = InResourceData.Width;
-	NewFrame.Stride = InResourceData.BytesPerRow;
-	NewFrame.VideoBuffer = InResourceData.Buffer;
-	NewFrame.FrameIdentifier = InBaseData.SourceFrameNumberRenderThread;
-	if (RivermaxStream->PushVideoFrame(NewFrame) == false)
-	{
-		UE_LOG(LogRivermaxMedia, Verbose, TEXT("Failed to pushed captured frame"));
-	}
+	OnFrameCapturedInternal_AnyThread(InBaseData, InUserData, InResourceData.Buffer, InResourceData.Width, InResourceData.Height, InResourceData.BytesPerRow);
 }
 
 FIntPoint URivermaxMediaCapture::GetCustomOutputSize(const FIntPoint& InSize) const
diff --git a/Engine/Plugins/VirtualProduction/Rivermax/RivermaxMedia/Source/RivermaxMedia/Private/RivermaxTimecodeProvider.cpp b/Engine/Plugins/VirtualProduction/Rivermax/RivermaxMedia/Source/RivermaxMedia/Private/RivermaxTimecodeProvider.cpp
new file mode 100644
--- /dev/null
+++ b/Engine/Plugins/VirtualProduction/Rivermax/RivermaxMedia/Source/RivermaxMedia/Private/RivermaxTimecodeProvider.cpp
@@ -0,0 +1,91 @@
+// Copyright Epic Games, Inc. All Rights Reserved.
+
+#include "RivermaxTimecodeProvider.h"
+
+#include "IRivermaxCoreModule.h"
+#include "IRivermaxManager.h"
+#include "RivermaxMediaLog.h"
+#include "RivermaxPTPUtils.h"
+
+#define LOCTEXT_NAMESPACE "RivermaxTimecodeProvider"
+
+
+bool URivermaxTimecodeProvider::FetchTimecode(FQualifiedFrameTime& OutFrameTime)
+{
+	IRivermaxCoreModule* RivermaxModule = FModuleManager::GetModulePtr<IRivermaxCoreModule>("RivermaxCore");
+	if (RivermaxModule && RivermaxModule->GetRivermaxManager())
+	{
+		// Get rivermax clock time, and truncate to timespan tick's resolution (100ns / tick)
+		const uint64 CurrentTime = RivermaxModule->GetRivermaxManager()->GetTime();
+		FTimespan CurrentTimespan = FTimespan(CurrentTime / ETimespan::NanosecondsPerTick);
+
+		// Adjust for daylight saving that might be required
+		CurrentTimespan -= FTimespan(DaylightSavingTimeHourOffset, 0, 0);
+
+		// Convert from TAI PTP Time to UTC
+		CurrentTimespan -= FTimespan(0, 0, UTCSecondsOffset);
+
+		constexpr bool bRollOver = true;
+		FTimecode Timecode = FTimecode::FromTimespan(CurrentTimespan, FrameRate, bRollOver);
+		OutFrameTime = FQualifiedFrameTime(Timecode, FrameRate);
+
+		// Adjust timecode to match LTC Timecode that should be one frame late because of LTC transport. 
+		OutFrameTime.Time.FrameNumber -= FFrameNumber(PTPToLTCTimecodeFrameOffset);
+
+		return true;
+	}
+
+	return false;
+}
+
+bool URivermaxTimecodeProvider::Initialize(class UEngine* InEngine)
+{
+	State = ETimecodeProviderSynchronizationState::Closed;
+
+	IRivermaxCoreModule* RivermaxModule = FModuleManager::GetModulePtr<IRivermaxCoreModule>("RivermaxCore");
+	if(RivermaxModule == nullptr || RivermaxModule->IsAvailable() == false)
+	{
+		UE_LOG(LogRivermaxMedia, Warning, TEXT("Can't initialize Rivermax TimecodeProvider, Rivermax Core Module isn't available."));
+		State = ETimecodeProviderSynchronizationState::Error;
+		return false;
+	}
+
+	TSharedPtr<UE::RivermaxCore::IRivermaxManager> Manager = RivermaxModule->GetRivermaxManager();
+	if(!Manager->ValidateLibraryIsLoaded())
+	{
+		UE_LOG(LogRivermaxMedia, Warning, TEXT("Can't initialize Rivermax TimecodeProvider, library isn't initialized."));
+		State = ETimecodeProviderSynchronizationState::Error;
+		return false;
+	}
+
+	if(RivermaxModule->GetRivermaxManager()->GetDevices().Num() <= 0)
+	{
+		UE_LOG(LogRivermaxMedia, Warning, TEXT("Can't initialize Rivermax TimecodeProvider, no compatible devices were found."));
+		State = ETimecodeProviderSynchronizationState::Error;
+		return false;
+	}
+
+	const ERivermaxTimeSource TimeSource = RivermaxModule->GetRivermaxManager()->GetTimeSource();
+	if ((TimeSource != ERivermaxTimeSource::PTP) && (TimeSource != ERivermaxTimeSource::System))
+	{
+		UE_LOG(LogRivermaxMedia, Warning, TEXT("Can't initialize Rivermax TimecodeProvider, Rivermax clock type isn't supported."));
+		State = ETimecodeProviderSynchronizationState::Error;
+		return false;
+	}
+
+	if (TimeSource == ERivermaxTimeSource::System)
+	{
+		UE_LOG(LogRivermaxMedia, Warning, TEXT("Rivermax clock is using system time. Make sure it's synchronized to PTP to get expected results."));
+	}
+
+	State = ETimecodeProviderSynchronizationState::Synchronized;
+
+	return true;
+}
+
+void URivermaxTimecodeProvider::Shutdown(class UEngine* InEngine)
+{
+	State = ETimecodeProviderSynchronizationState::Closed;
+}
+
+#undef LOCTEXT_NAMESPACE
diff --git a/Engine/Plugins/VirtualProduction/Rivermax/RivermaxMedia/Source/RivermaxMedia/Public/RivermaxMediaCapture.h b/Engine/Plugins/VirtualProduction/Rivermax/RivermaxMedia/Source/RivermaxMedia/Public/RivermaxMediaCapture.h
--- a/Engine/Plugins/VirtualProduction/Rivermax/RivermaxMedia/Source/RivermaxMedia/Public/RivermaxMediaCapture.h
+++ b/Engine/Plugins/VirtualProduction/Rivermax/RivermaxMedia/Source/RivermaxMedia/Public/RivermaxMediaCapture.h
@@ -38,6 +38,7 @@ protected:
 	virtual void StopCaptureImpl(bool bAllowPendingFrameToBeProcess) override;
 	virtual bool ShouldCaptureRHIResource() const override;
 	virtual void OnFrameCaptured_RenderingThread(const FCaptureBaseData& InBaseData, TSharedPtr<FMediaCaptureUserData, ESPMode::ThreadSafe> InUserData, void* InBuffer, int32 Width, int32 Height, int32 BytesPerRow) override;
+	virtual void OnRHIResourceCaptured_RenderingThread(const FCaptureBaseData& InBaseData, TSharedPtr<FMediaCaptureUserData, ESPMode::ThreadSafe> InUserData, FBufferRHIRef InBuffer) override;
 	virtual void OnRHIResourceCaptured_AnyThread(const FCaptureBaseData& InBaseData, TSharedPtr<FMediaCaptureUserData, ESPMode::ThreadSafe> InUserData, FBufferRHIRef InBuffer) override;
 	virtual void OnFrameCaptured_AnyThread(const FCaptureBaseData& InBaseData, TSharedPtr<FMediaCaptureUserData, ESPMode::ThreadSafe> InUserData, const FMediaCaptureResourceData& InResourceData) override;
 	virtual bool SupportsAnyThreadCapture() const override
@@ -78,6 +79,12 @@ private:
 
 	/** Enqueues a RHI lambda to reserve a spot for the next frame to capture */
 	void AddFrameReservationPass(FRDGBuilder& GraphBuilder);
+
+	/** Common method called for non gpudirect route when a frame is captured, either from render thread or any thread  */
+	void OnFrameCapturedInternal_AnyThread(const FCaptureBaseData& InBaseData, TSharedPtr<FMediaCaptureUserData, ESPMode::ThreadSafe> InUserData, void* InBuffer, int32 Width, int32 Height, int32 BytesPerRow);
+
+	/** Common method called for gpudirect route when a frame is captured, either from render thread or any thread  */
+	void OnRHIResourceCapturedInternal_AnyThread(const FCaptureBaseData& InBaseData, TSharedPtr<FMediaCaptureUserData, ESPMode::ThreadSafe> InUserData, FBufferRHIRef InBuffer);
 private:
 
 	/** Instance of the rivermax stream opened for this capture */
diff --git a/Engine/Plugins/VirtualProduction/Rivermax/RivermaxMedia/Source/RivermaxMedia/Public/RivermaxTimecodeProvider.h b/Engine/Plugins/VirtualProduction/Rivermax/RivermaxMedia/Source/RivermaxMedia/Public/RivermaxTimecodeProvider.h
new file mode 100644
--- /dev/null
+++ b/Engine/Plugins/VirtualProduction/Rivermax/RivermaxMedia/Source/RivermaxMedia/Public/RivermaxTimecodeProvider.h
@@ -0,0 +1,73 @@
+// Copyright Epic Games, Inc. All Rights Reserved.
+
+#pragma once
+
+#include "GenlockedTimecodeProvider.h"
+
+#include "RivermaxTimecodeProvider.generated.h"
+
+
+class UEngine;
+
+
+/**
+ * Timecode provider based on time returned by Rivermax GetTime API
+ * Rivermax clock is configured using Project Setting.
+ * If using multiple interfaces, make sure to configure the interface address you wish to use.
+ * 
+ * PTP time is returned in International Atomic Time referential (TAI) as opposed to Coordinated Universal Time (UTC)
+ * Timecode going out of GrandMaster is generally using UTC time reference so to generate an engine timecode
+ * that aligns with devices receiving LTC timecode, we need to convert PTP TAI time to UTC time. 
+ * At the moment, that offset is fixed at 37 seconds. 
+ * 
+ * GrandMasters can also output timecode with daylight saving time (DST) offsets which would not be affecting PTP time reference
+ * If devices receive timecode with DST offsets, engine timecode will need to be adjusted to correctly align.
+ */
+UCLASS(Blueprintable, editinlinenew)
+class RIVERMAXMEDIA_API URivermaxTimecodeProvider : public UGenlockedTimecodeProvider
+{
+public:
+	GENERATED_BODY()
+
+	//~ UTimecodeProvider interface
+	virtual bool FetchTimecode(FQualifiedFrameTime& OutFrameTime) override;
+	virtual ETimecodeProviderSynchronizationState GetSynchronizationState() const override { return State; }
+	virtual bool Initialize(class UEngine* InEngine) override;
+	virtual void Shutdown(class UEngine* InEngine) override;
+	virtual bool SupportsAutoDetected() const override
+	{
+		return false;
+	}
+
+public:
+
+	/** Frame Rate to provide timecode in */
+	UPROPERTY(EditAnywhere, Category = "Genlock options")
+	FFrameRate FrameRate = FFrameRate(24, 1);
+
+	/** 
+	 * Frame offset subtracted to converted timecode to match LTC timecode due to transportation of timecode
+	 * For LTC, it takes a full frame to transport the new value but for PTP, this doesn't happen, time is local.
+	 * If a timecoded video feed is coming into Unreal, if PTP timecode isn't offset, it will be leading by 1 frame.
+	 * A value of 1 means PTP timecode is 1 frame ahead so we substract 1 frame to the result
+	 * A value of -1 would mean PTP is late by one frame compared to LTC so we add a frame to the result.
+	 */
+	UPROPERTY(EditAnywhere, AdvancedDisplay, Category = "Genlock options")
+	int32 PTPToLTCTimecodeFrameOffset = 1;
+
+	/** TAI to UTC time offset. TAI is currently 37 seconds ahead so we substract 37 seconds to PTP time to convert to UTC. */
+	UPROPERTY(EditAnywhere, AdvancedDisplay, Category = "Genlock options")
+	int32 UTCSecondsOffset = 37;
+
+	/** 
+	 * If timecode generated by grand master has daylight saving applied, add offset here in order to match PTP based timecode with devices getting LTC timecode 
+	 * A positive value will be subtracted from the timecode generated from PTP time.
+	 */
+	UPROPERTY(EditAnywhere, AdvancedDisplay, Category = "Genlock options")
+	int32 DaylightSavingTimeHourOffset = 0;
+
+private:
+
+	/** The current SynchronizationState of the TimecodeProvider*/
+	ETimecodeProviderSynchronizationState State = ETimecodeProviderSynchronizationState::Closed;
+};
