- Removing unacceptable word in comments

diff --git a/Engine/Plugins/VirtualProduction/Rivermax/RivermaxMedia/Source/RivermaxMedia/Public/RivermaxTimecodeProvider.h b/Engine/Plugins/VirtualProduction/Rivermax/RivermaxMedia/Source/RivermaxMedia/Public/RivermaxTimecodeProvider.h
--- a/Engine/Plugins/VirtualProduction/Rivermax/RivermaxMedia/Source/RivermaxMedia/Public/RivermaxTimecodeProvider.h
+++ b/Engine/Plugins/VirtualProduction/Rivermax/RivermaxMedia/Source/RivermaxMedia/Public/RivermaxTimecodeProvider.h
@@ -16,11 +16,11 @@ class UEngine;
  * If using multiple interfaces, make sure to configure the interface address you wish to use.
  * 
  * PTP time is returned in International Atomic Time referential (TAI) as opposed to Coordinated Universal Time (UTC)
- * Timecode going out of GrandMaster is generally using UTC time reference so to generate an engine timecode
+ * Timecode going out of the clock generator is generally using UTC time reference so to generate an engine timecode
  * that aligns with devices receiving LTC timecode, we need to convert PTP TAI time to UTC time. 
  * At the moment, that offset is fixed at 37 seconds. 
  * 
- * GrandMasters can also output timecode with daylight saving time (DST) offsets which would not be affecting PTP time reference
+ * Clock generators can also output timecode with daylight saving time (DST) offsets which would not be affecting PTP time reference
  * If devices receive timecode with DST offsets, engine timecode will need to be adjusted to correctly align.
  */
 UCLASS(Blueprintable, editinlinenew)
@@ -60,7 +60,7 @@ public:
 	int32 UTCSecondsOffset = 37;
 
 	/** 
-	 * If timecode generated by grand master has daylight saving applied, add offset here in order to match PTP based timecode with devices getting LTC timecode 
+	 * If timecode generated by the clock generator has daylight saving applied, add offset here in order to match PTP based timecode with devices getting LTC timecode 
 	 * A positive value will be subtracted from the timecode generated from PTP time.
 	 */
 	UPROPERTY(EditAnywhere, AdvancedDisplay, Category = "Genlock options")
