- Adding phase shift / delay capability to rivermax / ptp custom timestep

diff --git a/Engine/Plugins/VirtualProduction/Rivermax/RivermaxMedia/Source/RivermaxMedia/Private/RivermaxCustomTimeStep.cpp b/Engine/Plugins/VirtualProduction/Rivermax/RivermaxMedia/Source/RivermaxMedia/Private/RivermaxCustomTimeStep.cpp
--- a/Engine/Plugins/VirtualProduction/Rivermax/RivermaxMedia/Source/RivermaxMedia/Private/RivermaxCustomTimeStep.cpp
+++ b/Engine/Plugins/VirtualProduction/Rivermax/RivermaxMedia/Source/RivermaxMedia/Private/RivermaxCustomTimeStep.cpp
@@ -185,7 +185,8 @@ bool URivermaxCustomTimeStep::WaitForNextFrame()
 	IRivermaxCoreModule& RivermaxModule = FModuleManager::GetModuleChecked<IRivermaxCoreModule>("RivermaxCore");
 	const uint64 CurrentPTPTime = RivermaxModule.GetRivermaxManager()->GetTime();
 	const double CurrentPlatformTime = FPlatformTime::Seconds();
-	const uint64 TargetTimeNanosec = UE::RivermaxCore::GetNextAlignmentPoint(CurrentPTPTime, FrameRate);
+	const uint64 DelayNanosec = AlignmentPointDelayMS * 1E6;
+	const uint64 TargetTimeNanosec = UE::RivermaxCore::GetNextAlignmentPoint(CurrentPTPTime, FrameRate) + DelayNanosec;
 
 	double TimeLeftNanosec = 0.0;
 	if(TargetTimeNanosec >= CurrentPTPTime)
diff --git a/Engine/Plugins/VirtualProduction/Rivermax/RivermaxMedia/Source/RivermaxMedia/Public/RivermaxCustomTimeStep.h b/Engine/Plugins/VirtualProduction/Rivermax/RivermaxMedia/Source/RivermaxMedia/Public/RivermaxCustomTimeStep.h
--- a/Engine/Plugins/VirtualProduction/Rivermax/RivermaxMedia/Source/RivermaxMedia/Public/RivermaxCustomTimeStep.h
+++ b/Engine/Plugins/VirtualProduction/Rivermax/RivermaxMedia/Source/RivermaxMedia/Public/RivermaxCustomTimeStep.h
@@ -58,6 +58,10 @@ public:
 	UPROPERTY(EditAnywhere, Category = "Genlock options")
 	FFrameRate FrameRate = FFrameRate(24,1);
 
+	/** Delay from alignment point at which we tick the engine, in milliseconds  */
+	UPROPERTY(EditAnywhere, Category = "Genlock options")
+	double AlignmentPointDelayMS = 0;
+
 	/** When enabled, will warn for skipped frames when engine is too slow */
 	UPROPERTY(EditAnywhere, Category = "Genlock options", meta = (DisplayName = "Display Dropped Frames Warning"), AdvancedDisplay)
 	bool bEnableOverrunDetection = true;
