Remove #pragma once from cpp files
#rnx

diff --git a/Engine/Plugins/VirtualProduction/Rivermax/RivermaxMedia/Source/RivermaxMediaEditor/Private/ModularFeatures/RivermaxMediaInitializerFeature.cpp b/Engine/Plugins/VirtualProduction/Rivermax/RivermaxMedia/Source/RivermaxMediaEditor/Private/ModularFeatures/RivermaxMediaInitializerFeature.cpp
--- a/Engine/Plugins/VirtualProduction/Rivermax/RivermaxMedia/Source/RivermaxMediaEditor/Private/ModularFeatures/RivermaxMediaInitializerFeature.cpp
+++ b/Engine/Plugins/VirtualProduction/Rivermax/RivermaxMedia/Source/RivermaxMediaEditor/Private/ModularFeatures/RivermaxMediaInitializerFeature.cpp
@@ -1,7 +1,5 @@
 // Copyright Epic Games, Inc. All Rights Reserved.
 
-#pragma once
-
 #include "ModularFeatures/RivermaxMediaInitializerFeature.h"
 
 #include "IRivermaxCoreModule.h"
