Rivermax:
- Added support for Color Encoding, Colorspace and OCIO.
- Rivermax is now merged with Media IO.
- Reduces the number of elements we wait for on various threads. Removed duplicate code.
- decoupled Framelocked samples from external media samples.

[RN]

diff --git a/Engine/Plugins/Media/MediaIOFramework/Source/MediaIOCore/Public/MediaIOCoreTextureSampleConverter.h b/Engine/Plugins/Media/MediaIOFramework/Source/MediaIOCore/Public/MediaIOCoreTextureSampleConverter.h
--- a/Engine/Plugins/Media/MediaIOFramework/Source/MediaIOCore/Public/MediaIOCoreTextureSampleConverter.h
+++ b/Engine/Plugins/Media/MediaIOFramework/Source/MediaIOCore/Public/MediaIOCoreTextureSampleConverter.h
@@ -33,7 +33,7 @@ public:
 	virtual uint32 GetConverterInfoFlags() const override;
 	//~ End IMediaTextureSampleConverter interface
 
-private:
+protected:
 
 	/** Proxy sample for JITR */
 	TWeakPtr<FMediaIOCoreTextureSampleBase> JITRProxySample;
diff --git a/Engine/Plugins/VirtualProduction/Rivermax/RivermaxCore/Shaders/Private/RivermaxShaders.usf b/Engine/Plugins/VirtualProduction/Rivermax/RivermaxCore/Shaders/Private/RivermaxShaders.usf
--- a/Engine/Plugins/VirtualProduction/Rivermax/RivermaxCore/Shaders/Private/RivermaxShaders.usf
+++ b/Engine/Plugins/VirtualProduction/Rivermax/RivermaxCore/Shaders/Private/RivermaxShaders.usf
@@ -296,10 +296,6 @@ void RGB10BitToRGBACS(uint3 ThreadId : SV_DispatchThreadID, uint GroupIndex : SV
 
 	for (uint PixelIndex = 0; PixelIndex < OutputTexturePixelsPerThread; PixelIndex++)
 	{
-	#if DO_SRGB_TO_LINEAR
-		PixelData[PixelIndex].RGB.xyz = sRGBToLinear(PixelData[PixelIndex].RGB.xyz);
-	#endif
-		
 		OutTexture[PixelData[PixelIndex].UV] = PixelData[PixelIndex].RGB;
 	}
 }
@@ -391,10 +387,6 @@ void RGB8BitToRGBA8CS(uint3 ThreadId : SV_DispatchThreadID, uint GroupIndex : SV
 
 	for (uint PixelIndex = 0; PixelIndex < OutputTexturePixelsPerThread; PixelIndex++)
 	{
-	#if DO_SRGB_TO_LINEAR
-		PixelData[PixelIndex].RGB.xyz = sRGBToLinear(PixelData[PixelIndex].RGB.xyz);
-	#endif
-		
 		OutTexture[PixelData[PixelIndex].UV] = PixelData[PixelIndex].RGB;
 	}
 }
@@ -506,10 +498,6 @@ void RGB12BitToRGBACS(uint3 ThreadId : SV_DispatchThreadID, uint GroupIndex : SV
 
 	for (uint PixelIndex = 0; PixelIndex < OutputTexturePixelsPerThread; PixelIndex++)
 	{
-	#if DO_SRGB_TO_LINEAR
-		PixelData[PixelIndex].RGB.xyz = sRGBToLinear(PixelData[PixelIndex].RGB.xyz);
-	#endif
-		
 		OutTexture[PixelData[PixelIndex].UV] = PixelData[PixelIndex].RGB;
 	}
 }
@@ -627,10 +615,6 @@ void RGB16fBitToRGBACS(uint3 ThreadId : SV_DispatchThreadID, uint GroupIndex : S
 	
 	for (uint PixelIndex = 0; PixelIndex < OutputTexturePixelsPerThread; PixelIndex++)
 	{
-	#if DO_SRGB_TO_LINEAR
-		PixelData[PixelIndex].RGB.xyz = sRGBToLinear(PixelData[PixelIndex].RGB.xyz);
-	#endif
-		
 		OutTexture[PixelData[PixelIndex].UV] = PixelData[PixelIndex].RGB;
 	}
 }
@@ -681,11 +665,7 @@ void YUV8Bit422ToRGBACS(uint3 ThreadId : SV_DispatchThreadID, uint GroupIndex :
 	PixelData[0].YUVY.z = ((InputYUV4228bitBuffer[BufferIndex].DWord0 >> 16) & 0xFF) * Mul;
 	PixelData[0].YUVY.w = ((InputYUV4228bitBuffer[BufferIndex].DWord0 >> 24) & 0xFF) * Mul;
 
-	#if DO_SRGB_TO_LINEAR
-	int DoSRGBToLinear = 1;
-	#else
 	int DoSRGBToLinear = 0;
-	#endif
 
 	for (uint PixelIndex = 0; PixelIndex < PixelContainerSize; PixelIndex++)
 	{
@@ -754,11 +734,8 @@ void YUV10Bit422ToRGBACS(uint3 ThreadId : SV_DispatchThreadID, uint GroupIndex :
 	PixelData[3].YUVY.z = ((((InputYCbCrBuffer[BufferIndex].DWord4 >>  8)  << 6) | ((InputYCbCrBuffer[BufferIndex].DWord4 >> 18) & 0x03F)) & 0x3FF) * Mul;
 	PixelData[3].YUVY.w = ((((InputYCbCrBuffer[BufferIndex].DWord4 >> 16)  << 8) | ((InputYCbCrBuffer[BufferIndex].DWord4 >> 24) & 0x0FF)) & 0x3FF) * Mul;
 
-	#if DO_SRGB_TO_LINEAR
-	int DoSRGBToLinear = 1;
-	#else
+
 	int DoSRGBToLinear = 0;
-	#endif
 
 	for (uint PixelIndex = 0; PixelIndex < PixelContainerSize; PixelIndex++)
 	{
diff --git a/Engine/Plugins/VirtualProduction/Rivermax/RivermaxCore/Source/RivermaxCore/Private/RTPHeader.cpp b/Engine/Plugins/VirtualProduction/Rivermax/RivermaxCore/Source/RivermaxCore/Private/RTPHeader.cpp
--- a/Engine/Plugins/VirtualProduction/Rivermax/RivermaxCore/Source/RivermaxCore/Private/RTPHeader.cpp
+++ b/Engine/Plugins/VirtualProduction/Rivermax/RivermaxCore/Source/RivermaxCore/Private/RTPHeader.cpp
@@ -78,7 +78,7 @@ namespace UE::RivermaxCore::Private
 		}
 
 		// Pretty sure some data needs to be swapped but can't validate that until we have other hardware generating data
-		SequencerNumber = (ByteSwap((uint16)RawHeader.ExtendedSequenceNumber) << 16) | ByteSwap((uint16)RawHeader.SequenceNumber);
+		SequenceNumber = (ByteSwap((uint16)RawHeader.ExtendedSequenceNumber) << 16) | ByteSwap((uint16)RawHeader.SequenceNumber);
 		Timestamp = ByteSwap(RawHeader.Timestamp);
 		bIsMarkerBit = RawHeader.MarkerBit;
 
diff --git a/Engine/Plugins/VirtualProduction/Rivermax/RivermaxCore/Source/RivermaxCore/Private/RTPHeader.h b/Engine/Plugins/VirtualProduction/Rivermax/RivermaxCore/Source/RivermaxCore/Private/RTPHeader.h
--- a/Engine/Plugins/VirtualProduction/Rivermax/RivermaxCore/Source/RivermaxCore/Private/RTPHeader.h
+++ b/Engine/Plugins/VirtualProduction/Rivermax/RivermaxCore/Source/RivermaxCore/Private/RTPHeader.h
@@ -134,7 +134,7 @@ namespace UE::RivermaxCore::Private
 		uint16 GetLastRowNumber() const;
 
 		/** Sequence number including extension if present */
-		uint32 SequencerNumber = 0;
+		uint32 SequenceNumber = 0;
 
 		/** Timestamp of frame in the specified clock resolution. Video is typically 90kHz */
 		uint32 Timestamp = 0;
diff --git a/Engine/Plugins/VirtualProduction/Rivermax/RivermaxCore/Source/RivermaxCore/Private/RivermaxInputStream.cpp b/Engine/Plugins/VirtualProduction/Rivermax/RivermaxCore/Source/RivermaxCore/Private/RivermaxInputStream.cpp
--- a/Engine/Plugins/VirtualProduction/Rivermax/RivermaxCore/Source/RivermaxCore/Private/RivermaxInputStream.cpp
+++ b/Engine/Plugins/VirtualProduction/Rivermax/RivermaxCore/Source/RivermaxCore/Private/RivermaxInputStream.cpp
@@ -9,6 +9,7 @@
 #include "IRivermaxCoreModule.h"
 #include "IRivermaxManager.h"
 #include "Misc/ByteSwap.h"
+#include "RenderGraphUtils.h"
 #include "RivermaxLog.h"
 #include "RivermaxPTPUtils.h"
 #include "RivermaxTracingUtils.h"
@@ -314,9 +315,10 @@ namespace UE::RivermaxCore::Private
 
 	void FRivermaxInputStream::ParseChunks(const rmx_input_completion* Completion)
 	{
-		const size_t ChunkCount = rmx_input_get_completion_chunk_size(Completion);
+		const size_t PacketCount = rmx_input_get_completion_chunk_size(Completion);
 
-		for (uint64 StrideIndex = 0; StrideIndex < ChunkCount; ++StrideIndex)
+		// We need to get all samples once per group of packets and store against its frame number.
+		for (uint64 StrideIndex = 0; StrideIndex < PacketCount; ++StrideIndex)
 		{
 			++StreamStats.ChunksReceived;
 
@@ -354,7 +356,6 @@ namespace UE::RivermaxCore::Private
 				if (RawRTPHeaderPtr.Version == 2)
 				{
 					FRTPHeader RTPHeader(RawRTPHeaderPtr);
-					
 					// Add trace for the first packet of a frame to help visualize reception of a full frame in time
 					if (bIsFirstPacketReceived == false)
 					{
@@ -404,11 +405,18 @@ namespace UE::RivermaxCore::Private
 					
 					UpdateFrameTracking(RTPHeader);
 
+					// Get video sample to write into (once per frame)
+					if (!StreamData.CurrentSamples.Contains(IRivermaxSample::ESampleType::Video) && State != EReceptionState::FrameError)
+					{
+						TSharedPtr<IRivermaxVideoSample> Sample = GetVideoSampleForReception(RTPHeader);
+						StreamData.CurrentSamples.Add(IRivermaxSample::ESampleType::Video, StaticCastSharedPtr<IRivermaxSample>(Sample));
+					}
+
 					switch (State)
 					{
 					case EReceptionState::Receiving:
 					{
-						FrameReceptionState(RTPHeader, DataPtr);
+						FrameReceptionState(RTPHeader, DataPtr, StreamData.CurrentSamples[IRivermaxSample::ESampleType::Video]);
 						break;
 					}
 					case EReceptionState::WaitingForMarker:
@@ -439,7 +447,7 @@ namespace UE::RivermaxCore::Private
 		}
 	}
 
-	bool FRivermaxInputStream::PrepareNextFrame(const FRTPHeader& RTPHeader)
+	TSharedPtr<IRivermaxVideoSample> FRivermaxInputStream::GetVideoSampleForReception(const FRTPHeader& RTPHeader)
 	{
 		using namespace UE::RivermaxCore::Private::Utils;
 
@@ -453,46 +461,48 @@ namespace UE::RivermaxCore::Private
 		const uint32 FrameSize = PixelCount / FormatInfo.PixelGroupCoverage * FormatInfo.PixelGroupSize;
 		Descriptor.VideoBufferSize = FrameSize;
 		FRivermaxInputVideoFrameRequest Request;
-		Listener->OnVideoFrameRequested(Descriptor, Request);
+		TSharedPtr<IRivermaxVideoSample> OutSample = Listener->OnVideoFrameRequested(Descriptor);
 
 		// Reset current frame to know when we have a valid one
-		StreamData.CurrentFrame = nullptr;
-		if (bIsUsingGPUDirect)
+		StreamData.CurrentFrameVideoBuffer = nullptr;
+		if (OutSample.IsValid())
 		{
-			if(Request.GPUBuffer)
+			if (bIsUsingGPUDirect)
 			{
-				StreamData.CurrentFrame = GetMappedBuffer(Request.GPUBuffer);
+				FBufferRHIRef RHIBuffer = OutSample->GetGPUBuffer()->GetRHI();
+				if (RHIBuffer.IsValid())
+				{
+					StreamData.CurrentFrameVideoBuffer = GetMappedBuffer(RHIBuffer);
+				}
 			}
-		}
-		else
-		{	
-			if (Request.VideoBuffer)
+			else
 			{
-				StreamData.CurrentFrame = Request.VideoBuffer;
+				StreamData.CurrentFrameVideoBuffer = OutSample->GetVideoBufferRawPtr(Descriptor.VideoBufferSize);
 			}
 		}
 
 		// Verify if we were able to request a valid frame. If engine is blocked, it could happen that there is none available 
-		if (StreamData.CurrentFrame == nullptr)
+		if (StreamData.CurrentFrameVideoBuffer == nullptr)
 		{
 			// If we failed getting one, reset the valid first frame received and wait for the next one
 			UE_LOG(LogRivermax, Verbose, TEXT("Could not get a new frame for incoming frame with timestamp %u and frame number %u"), RTPHeader.Timestamp, Descriptor.FrameNumber);
-			Listener->OnVideoFrameReceptionError(Descriptor);
+			Listener->OnVideoFrameReceptionError(OutSample);
 			State = EReceptionState::FrameError;
 			FrameErrorState(RTPHeader);
-			return false;
+		}
+		else
+		{
+			StreamData.WritingOffset = 0;
+			StreamData.ReceivedSize = 0;
+			StreamData.ExpectedSize = Descriptor.VideoBufferSize;
+			StreamData.DeviceWritePointerOne = nullptr;
+			StreamData.SizeToWriteOne = 0;
+			StreamData.DeviceWritePointerTwo = nullptr;
+			StreamData.SizeToWriteTwo = 0;
+			bIsFirstPacketReceived = false;
 		}
 
-		StreamData.WritingOffset = 0;
-		StreamData.ReceivedSize = 0;
-		StreamData.ExpectedSize = Descriptor.VideoBufferSize;
-		StreamData.DeviceWritePointerOne = nullptr;
-		StreamData.SizeToWriteOne = 0;
-		StreamData.DeviceWritePointerTwo = nullptr;
-		StreamData.SizeToWriteTwo = 0;
-		bIsFirstPacketReceived = false;
-
-		return true;
+		return OutSample;
 	}
 
 	void FRivermaxInputStream::LogStats()
@@ -815,14 +825,6 @@ namespace UE::RivermaxCore::Private
 
 	void FRivermaxInputStream::ProcessSRD(const FRTPHeader& RTPHeader, const uint8* DataPtr)
 	{
-		if (StreamData.CurrentFrame == nullptr)
-		{
-			if (PrepareNextFrame(RTPHeader) == false)
-			{
-				return;
-			}
-		}
-
 		if (RTPHeader.SRD1.Length <= 0)
 		{
 			return;
@@ -865,7 +867,7 @@ namespace UE::RivermaxCore::Private
 		}
 		else
 		{
-			uint8* WriteBuffer = reinterpret_cast<uint8*>(StreamData.CurrentFrame);
+			uint8* WriteBuffer = reinterpret_cast<uint8*>(StreamData.CurrentFrameVideoBuffer);
 			FMemory::Memcpy(&WriteBuffer[StreamData.WritingOffset], &DataPtr[DataOffset], RTPHeader.SRD1.Length);
 			StreamData.WritingOffset += RTPHeader.SRD1.Length;
 
@@ -880,22 +882,13 @@ namespace UE::RivermaxCore::Private
 		StreamData.ReceivedSize += PayloadSize;
 	}
 
-	void FRivermaxInputStream::ProcessLastSRD(const FRTPHeader& RTPHeader, const uint8* DataPtr)
+	void FRivermaxInputStream::ProcessLastSRD(const FRTPHeader& RTPHeader, const uint8* DataPtr, TSharedPtr<IRivermaxSample> InSample)
 	{
-		FRivermaxInputVideoFrameDescriptor Descriptor;
-		Descriptor.Width = StreamResolution.X;
-		Descriptor.Height = StreamResolution.Y;
-		Descriptor.PixelFormat = Options.PixelFormat;
-		Descriptor.Timestamp = RTPHeader.Timestamp;
-		const uint32 PixelCount = StreamResolution.X * StreamResolution.Y;
-		Descriptor.VideoBufferSize = PixelCount / FormatInfo.PixelGroupCoverage * FormatInfo.PixelGroupSize;
-		Descriptor.FrameNumber = Utils::TimestampToFrameNumber(RTPHeader.Timestamp, Options.FrameRate);
-
 		if (StreamData.ReceivedSize == StreamData.ExpectedSize)
 		{
 			++StreamStats.FramesReceived;
-			
-			TRACE_CPUPROFILER_EVENT_SCOPE_TEXT(*FRivermaxTracingUtils::RmaxInReceivedFrameTraceEvents[Descriptor.FrameNumber % 10]);
+			uint32 FrameNumber = Utils::TimestampToFrameNumber(RTPHeader.Timestamp, Options.FrameRate);
+			TRACE_CPUPROFILER_EVENT_SCOPE_TEXT(*FRivermaxTracingUtils::RmaxInReceivedFrameTraceEvents[FrameNumber % 10]);
 
 			if (bIsUsingGPUDirect)
 			{
@@ -911,7 +904,7 @@ namespace UE::RivermaxCore::Private
 				FCUDAModule& CudaModule = FModuleManager::GetModuleChecked<FCUDAModule>("CUDA");
 				CUresult Result = CudaModule.DriverAPI()->cuCtxPushCurrent(CudaModule.GetCudaContext());
 
-				const CUdeviceptr DestinationGPUMemory = reinterpret_cast<CUdeviceptr>(StreamData.CurrentFrame);
+				const CUdeviceptr DestinationGPUMemory = reinterpret_cast<CUdeviceptr>(StreamData.CurrentFrameVideoBuffer);
 				const CUdeviceptr SourceGPUMemoryOne = reinterpret_cast<CUdeviceptr>(StreamData.DeviceWritePointerOne);
 
 				const uint32 NumSRDPartOne = StreamData.SizeToWriteOne / FrameDescriptionTracking.CommonPayloadSize;
@@ -940,7 +933,7 @@ namespace UE::RivermaxCore::Private
 				if (StreamData.DeviceWritePointerTwo != nullptr && StreamData.SizeToWriteTwo > 0)
 				{
 					StrideDescription.srcDevice = reinterpret_cast<CUdeviceptr>(StreamData.DeviceWritePointerTwo);
-					StrideDescription.dstDevice = reinterpret_cast<CUdeviceptr>(StreamData.CurrentFrame) + StreamData.SizeToWriteOne;
+					StrideDescription.dstDevice = reinterpret_cast<CUdeviceptr>(StreamData.CurrentFrameVideoBuffer) + StreamData.SizeToWriteOne;
 					StrideDescription.Height = NumSRDPartTwo;
 					Result = CudaModule.DriverAPI()->cuMemcpy2DAsync(&StrideDescription, reinterpret_cast<CUstream>(GPUStream));
 
@@ -963,7 +956,7 @@ namespace UE::RivermaxCore::Private
 				if (Result != CUDA_SUCCESS)
 				{
 					UE_LOG(LogRivermax, Warning, TEXT("Failed to copy received buffer to shared memory. Error: %d"), Result);
-					Listener->OnVideoFrameReceptionError(Descriptor);
+					Listener->OnVideoFrameReceptionError(StaticCastSharedPtr<IRivermaxVideoSample>(InSample));
 					State = EReceptionState::FrameError;
 					FrameErrorState(RTPHeader);
 					return;
@@ -1016,13 +1009,11 @@ namespace UE::RivermaxCore::Private
 			// No need to provide the new frame and prepare the next one if we are shutting down
 			if (bIsShuttingDown == false)
 			{
-				UE_LOG(LogRivermax, VeryVerbose, TEXT("RmaxRX frame number %u with timestamp %u."), Descriptor.FrameNumber, Descriptor.Timestamp);
-				
 				FRivermaxInputVideoFrameReception NewFrame;
-				NewFrame.VideoBuffer = reinterpret_cast<uint8*>(StreamData.CurrentFrame);
-				Listener->OnVideoFrameReceived(Descriptor, NewFrame);
-				StreamData.CurrentFrame = nullptr;
-
+				NewFrame.VideoBuffer = reinterpret_cast<uint8*>(StreamData.CurrentFrameVideoBuffer);
+				Listener->OnVideoFrameReceived(StaticCastSharedPtr<IRivermaxVideoSample>(InSample));
+				StreamData.CurrentFrameVideoBuffer = nullptr;
+				StreamData.CurrentSamples.Empty();
 				// Finished receiving a frame, move back to reception
 				State = EReceptionState::Receiving;
 			}
@@ -1031,7 +1022,7 @@ namespace UE::RivermaxCore::Private
 		{
 			UE_LOG(LogRivermax, Warning, TEXT("End of frame received but not enough data was received (missing %d). Expected %d but received (%d)"), StreamData.ExpectedSize - StreamData.ReceivedSize, StreamData.ExpectedSize, StreamData.ReceivedSize);
 			
-			Listener->OnVideoFrameReceptionError(Descriptor);
+			Listener->OnVideoFrameReceptionError(StaticCastSharedPtr<IRivermaxVideoSample>(InSample));
 			State = EReceptionState::FrameError;
 			FrameErrorState(RTPHeader);
 			
@@ -1039,24 +1030,14 @@ namespace UE::RivermaxCore::Private
 		}
 	}
 
-	void FRivermaxInputStream::FrameReceptionState(const FRTPHeader& RTPHeader, const uint8* DataPtr)
+	void FRivermaxInputStream::FrameReceptionState(const FRTPHeader& RTPHeader, const uint8* DataPtr, TSharedPtr<IRivermaxSample> InSample)
 	{
-		FRivermaxInputVideoFrameDescriptor Descriptor;
-		Descriptor.Width = StreamResolution.X;
-		Descriptor.Height = StreamResolution.Y;
-		Descriptor.PixelFormat = Options.PixelFormat;
-		Descriptor.Timestamp = RTPHeader.Timestamp;
-		const uint32 PixelCount = StreamResolution.X * StreamResolution.Y;
-		const uint32 FrameSize = PixelCount / FormatInfo.PixelGroupCoverage * FormatInfo.PixelGroupSize;
-		Descriptor.VideoBufferSize = FrameSize;
-		Descriptor.FrameNumber = Utils::TimestampToFrameNumber(RTPHeader.Timestamp, Options.FrameRate);
-
 		const uint64 LastSequenceNumberIncremented = StreamData.LastSequenceNumber + 1;
 
-		const uint64 LostPackets = ((uint64)RTPHeader.SequencerNumber + 0x100000000 - LastSequenceNumberIncremented) & 0xFFFFFFFF;
+		const uint64 LostPackets = ((uint64)RTPHeader.SequenceNumber + 0x100000000 - LastSequenceNumberIncremented) & 0xFFFFFFFF;
 		if (LostPackets > 0)
 		{
-			UE_LOG(LogRivermax, Warning, TEXT("Lost %llu packets during reception of packet %u"), LostPackets, Descriptor.FrameNumber);
+			UE_LOG(LogRivermax, Warning, TEXT("Lost %llu packets during reception of chunk"), LostPackets);
 			
 			StreamData.WritingOffset = 0;
 			StreamData.ReceivedSize = 0;
@@ -1064,13 +1045,13 @@ namespace UE::RivermaxCore::Private
 			++StreamStats.FramePacketLossCount;
 
 			// For now, if packets were lost, skip the incoming frame. We could improve that and have corrupted frames instead of skipping them but can be added later
-			Listener->OnVideoFrameReceptionError(Descriptor);
+			Listener->OnVideoFrameReceptionError(StaticCastSharedPtr<IRivermaxVideoSample>(InSample));
 			State = EReceptionState::FrameError;
 			FrameErrorState(RTPHeader);
 			return;
 		}
 
-		StreamData.LastSequenceNumber = RTPHeader.SequencerNumber;
+		StreamData.LastSequenceNumber = RTPHeader.SequenceNumber;
 
 		ProcessSRD(RTPHeader, DataPtr);
 
@@ -1081,7 +1062,7 @@ namespace UE::RivermaxCore::Private
 			StreamData.ReceivedSize = 0;
 			++StreamStats.BiggerFramesCount;
 
-			Listener->OnVideoFrameReceptionError(Descriptor);
+			Listener->OnVideoFrameReceptionError(StaticCastSharedPtr<IRivermaxVideoSample>(InSample));
 			State = EReceptionState::FrameError;
 			FrameErrorState(RTPHeader);
 			return;
@@ -1089,7 +1070,7 @@ namespace UE::RivermaxCore::Private
 
 		if (RTPHeader.bIsMarkerBit)
 		{
-			ProcessLastSRD(RTPHeader, DataPtr);
+			ProcessLastSRD(RTPHeader, DataPtr, InSample);
 
 			StreamStats.FramePacketLossCount = 0;
 			++StreamStats.EndOfFrameReceived;
@@ -1100,7 +1081,7 @@ namespace UE::RivermaxCore::Private
 	{
 		if (RTPHeader.bIsMarkerBit)
 		{
-			StreamData.LastSequenceNumber = RTPHeader.SequencerNumber;
+			StreamData.LastSequenceNumber = RTPHeader.SequenceNumber;
 			StreamData.WritingOffset = 0;
 			StreamData.ReceivedSize = 0;
 			StreamData.DeviceWritePointerOne = nullptr;
@@ -1109,7 +1090,8 @@ namespace UE::RivermaxCore::Private
 			StreamData.SizeToWriteTwo = 0;
 			bIsFirstPacketReceived = false;
 
-			StreamData.CurrentFrame = nullptr;
+			StreamData.CurrentFrameVideoBuffer = nullptr;
+			StreamData.CurrentSamples.Empty();
 			State = EReceptionState::Receiving;
 		}
 	}
@@ -1281,5 +1263,4 @@ namespace UE::RivermaxCore::Private
 }
 
 	
-	
 
diff --git a/Engine/Plugins/VirtualProduction/Rivermax/RivermaxCore/Source/RivermaxCore/Private/RivermaxInputStream.h b/Engine/Plugins/VirtualProduction/Rivermax/RivermaxCore/Source/RivermaxCore/Private/RivermaxInputStream.h
--- a/Engine/Plugins/VirtualProduction/Rivermax/RivermaxCore/Source/RivermaxCore/Private/RivermaxInputStream.h
+++ b/Engine/Plugins/VirtualProduction/Rivermax/RivermaxCore/Source/RivermaxCore/Private/RivermaxInputStream.h
@@ -72,7 +72,7 @@ namespace UE::RivermaxCore::Private
 	struct FInputStreamData
 	{
 		uint64 LastSequenceNumber = 0;
-		void* CurrentFrame = nullptr;
+		void* CurrentFrameVideoBuffer = nullptr;
 		uint32 WritingOffset = 0;
 		uint32 ReceivedSize = 0;
 		uint32 ExpectedSize = 0;
@@ -82,6 +82,9 @@ namespace UE::RivermaxCore::Private
 		const void* DeviceWritePointerTwo = nullptr;
 		uint32 SizeToWriteOne = 0;
 		uint32 SizeToWriteTwo = 0;
+
+		// A map of samples for the current frame. 
+		TMap<IRivermaxSample::ESampleType, TSharedPtr<IRivermaxSample>> CurrentSamples;
 	};
 
 	/** Incoming frame description tracker. Used to detect variable SRD length across a frame for now but should be used to help auto detect incoming frame */
@@ -157,7 +160,7 @@ namespace UE::RivermaxCore::Private
 	private:
 
 		/** Process reception state. i.e Extract RTP, tracks data received, etc... */
-		void FrameReceptionState(const FRTPHeader& RTPHeader, const uint8* DataPtr);
+		void FrameReceptionState(const FRTPHeader& RTPHeader, const uint8* DataPtr, TSharedPtr<IRivermaxSample> InSample);
 
 		/** State where we just wait to receive an end of frame in RTP. Used to align for the first frame and when an error occurs during reception */
 		void WaitForMarkerState(const FRTPHeader& RTPHeader);
@@ -178,10 +181,10 @@ namespace UE::RivermaxCore::Private
 		void ProcessSRD(const FRTPHeader& RTPHeader, const uint8* DataPtr);
 
 		/** Handles wrapping up of a frame when marker bit is detected (last SRD). Will notify listener of a new frame received */
-		void ProcessLastSRD(const FRTPHeader& RTPHeader, const uint8* DataPtr);
+		void ProcessLastSRD(const FRTPHeader& RTPHeader, const uint8* DataPtr, TSharedPtr<IRivermaxSample> InSample);
 
 		/** Requests a new frame from our listener. If it fails, it will fall in error state waiting for the next one */
-		bool PrepareNextFrame(const FRTPHeader& RTPHeader);
+		TSharedPtr<IRivermaxVideoSample> GetVideoSampleForReception(const FRTPHeader& RTPHeader);
 
 		/** If enabled, will printout reception stats to logs */
 		void LogStats();
diff --git a/Engine/Plugins/VirtualProduction/Rivermax/RivermaxCore/Source/RivermaxCore/Private/RivermaxPTPUtils.cpp b/Engine/Plugins/VirtualProduction/Rivermax/RivermaxCore/Source/RivermaxCore/Private/RivermaxPTPUtils.cpp
--- a/Engine/Plugins/VirtualProduction/Rivermax/RivermaxCore/Source/RivermaxCore/Private/RivermaxPTPUtils.cpp
+++ b/Engine/Plugins/VirtualProduction/Rivermax/RivermaxCore/Source/RivermaxCore/Private/RivermaxPTPUtils.cpp
@@ -4,6 +4,7 @@
 
 #include "Misc/FrameRate.h"
 #include "RivermaxTypes.h"
+#include "RivermaxUtils.h"
 
 namespace UE::RivermaxCore
 {
@@ -64,4 +65,9 @@ namespace UE::RivermaxCore
 
 		return FSec + FNanoSec + FractionalFrames;
 	}
+
+	RIVERMAXCORE_API uint64 ConvertRTPTimeStampToSecondsTruncated(const uint64 InTimestamp)
+	{
+		return InTimestamp / UE::RivermaxCore::Private::Utils::MediaClockSampleRate;
+	}
 }
\ No newline at end of file
diff --git a/Engine/Plugins/VirtualProduction/Rivermax/RivermaxCore/Source/RivermaxCore/Public/IRivermaxInputStream.h b/Engine/Plugins/VirtualProduction/Rivermax/RivermaxCore/Source/RivermaxCore/Public/IRivermaxInputStream.h
--- a/Engine/Plugins/VirtualProduction/Rivermax/RivermaxCore/Source/RivermaxCore/Public/IRivermaxInputStream.h
+++ b/Engine/Plugins/VirtualProduction/Rivermax/RivermaxCore/Source/RivermaxCore/Public/IRivermaxInputStream.h
@@ -6,6 +6,7 @@
 #include "RivermaxFormats.h"
 
 class FRHIBuffer;
+class FRDGPooledBuffer;
 
 namespace UE::RivermaxCore
 {
@@ -67,20 +68,108 @@ namespace UE::RivermaxCore
 		uint8* VideoBuffer = nullptr;
 	};
 
+	/** 
+	* All types of streams write into corresponding sample types. All sample types should inherit from this class so it can be correctly cast and identified.
+	*/
+	class RIVERMAXCORE_API IRivermaxSample
+	{
+	public:
+		virtual ~IRivermaxSample()
+		{
+
+		}
+		enum class ESampleState : uint8
+		{
+			// Sample is ready to be written into.
+			ReadyForReception,
+
+			// Error receiving sample.
+			ReceptionError
+		};
+
+		/**
+		* Enum that identifies the type of this sample such as Video or Audio so that it can be handled accordingly.
+		*/
+		enum class ESampleType : uint8
+		{
+			// 2110-20
+			Video,
+
+			// 2110-30
+			Audio,
+
+			// 2110-40
+			Anc,
+
+			// 2110-20 Sub raster.
+			KeyAndFill,
+
+			// Stub for itteration.
+			MAX
+		};
+
+	public:
+
+		/**
+		* Sets the state of the reception for this sample.
+		*/
+		virtual void SetReceptionState(ESampleState State)
+		{
+			FScopeLock Lock(&StateChangeCriticalSecion);
+			SampleState = State;
+		}
+
+		/**
+		* Gets reception state of this sample.
+		*/
+		virtual ESampleState GetReceptionState()
+		{
+			FScopeLock Lock(&StateChangeCriticalSecion);
+			return SampleState;
+		}
+
+	protected:
+		/** State of this sample. */
+		std::atomic<ESampleState> SampleState = ESampleState::ReadyForReception;
+
+		/** Type of the sample. */
+		std::atomic<ESampleType> SampleType = ESampleType::Video;
+
+		/** Critical section used when manipulating the received/skipped/rendered and other states in this or the inheriting classes. */
+		mutable FCriticalSection StateChangeCriticalSecion;
+	};
+
+	/**
+	* 2110-20 sample type interface.
+	*/
+	class RIVERMAXCORE_API IRivermaxVideoSample : public IRivermaxSample
+	{
+	public:
+		
+		/** Returns RDG allocated buffer */
+		virtual TRefCountPtr<FRDGPooledBuffer> GetGPUBuffer() const = 0;
+
+		/** Returns a pointer to the CPU accessible buffer for writing streams into. */
+		virtual uint8* GetVideoBufferRawPtr(uint32 VideoBufferSize) = 0;
+	};
+
+	/**
+	* A type of interface that is able to provide samples for data to be written into and react to stream events such as completion of the reception.
+	*/
 	class RIVERMAXCORE_API IRivermaxInputStreamListener
 	{
 	public:
 		/** Initialization completion callback with result */
 		virtual void OnInitializationCompleted(const FRivermaxInputInitializationResult& Result) = 0;
-	
+
 		/** Called when stream is ready to fill the next frame. Returns true if a frame was successfully requested */
-		virtual bool OnVideoFrameRequested(const FRivermaxInputVideoFrameDescriptor& FrameInfo, FRivermaxInputVideoFrameRequest& OutVideoFrameRequest) = 0;
-		
+		virtual TSharedPtr<IRivermaxVideoSample> OnVideoFrameRequested(const FRivermaxInputVideoFrameDescriptor& FrameInfo) = 0;
+
 		/** Called when a frame has been received */
-		virtual void OnVideoFrameReceived(const FRivermaxInputVideoFrameDescriptor& FrameInfo, const FRivermaxInputVideoFrameReception& ReceivedVideoFrame) = 0;
+		virtual void OnVideoFrameReceived(TSharedPtr<IRivermaxVideoSample> InReceivedVideoFrame) = 0;
 
 		/** Called when an error was encountered during frame reception */
-		virtual void OnVideoFrameReceptionError(const FRivermaxInputVideoFrameDescriptor& FrameInfo) {};
+		virtual void OnVideoFrameReceptionError(TSharedPtr<IRivermaxVideoSample> InVideoFrameSample) {};
 
 		/** Called when stream has encountered an error and has to stop */
 		virtual void OnStreamError() = 0;
@@ -89,6 +178,9 @@ namespace UE::RivermaxCore
 		virtual void OnVideoFormatChanged(const FRivermaxInputVideoFormatChangedInfo& NewFormatInfo) = 0;
 	};
 
+	/**
+	* Interface for initializing input stream from media player.
+	*/
 	class RIVERMAXCORE_API IRivermaxInputStream
 	{
 	public:
diff --git a/Engine/Plugins/VirtualProduction/Rivermax/RivermaxCore/Source/RivermaxCore/Public/RivermaxPTPUtils.h b/Engine/Plugins/VirtualProduction/Rivermax/RivermaxCore/Source/RivermaxCore/Public/RivermaxPTPUtils.h
--- a/Engine/Plugins/VirtualProduction/Rivermax/RivermaxCore/Source/RivermaxCore/Public/RivermaxPTPUtils.h
+++ b/Engine/Plugins/VirtualProduction/Rivermax/RivermaxCore/Source/RivermaxCore/Public/RivermaxPTPUtils.h
@@ -24,6 +24,12 @@ namespace UE::RivermaxCore
 	 * Returns current frame number for the incoming PTP time for the given frame rate.
 	 */
 	RIVERMAXCORE_API uint64 GetFrameNumber(const uint64 InPTPTimeNanosec, const FFrameRate& InRate);
+
+	/**
+	 * Converts an RTP timestamp in MediaClock period units to seconds
+	 * 2110-20 streams uses a standard media clock rate of 90kHz
+	 */
+	RIVERMAXCORE_API uint64 ConvertRTPTimeStampToSecondsTruncated(const uint64 InTimestamp);
 }
 
 
diff --git a/Engine/Plugins/VirtualProduction/Rivermax/RivermaxCore/Source/RivermaxCore/RivermaxCore.build.cs b/Engine/Plugins/VirtualProduction/Rivermax/RivermaxCore/Source/RivermaxCore/RivermaxCore.build.cs
--- a/Engine/Plugins/VirtualProduction/Rivermax/RivermaxCore/Source/RivermaxCore/RivermaxCore.build.cs
+++ b/Engine/Plugins/VirtualProduction/Rivermax/RivermaxCore/Source/RivermaxCore/RivermaxCore.build.cs
@@ -12,7 +12,6 @@ public class RivermaxCore : ModuleRules
 				"CoreUObject",
 				"DeveloperSettings",
 				"MediaAssets",
-				"MediaIOCore",
 				"RHI",
 				"RivermaxLib",
 				"TimeManagement"
@@ -24,7 +23,7 @@ public class RivermaxCore : ModuleRules
 				"D3D12RHI",
 				"Engine",
 				"Networking",
-				"RenderCore"
+				"RenderCore",
 			}
 		);
 
diff --git a/Engine/Plugins/VirtualProduction/Rivermax/RivermaxCore/Source/RivermaxRendering/Public/RivermaxShaders.h b/Engine/Plugins/VirtualProduction/Rivermax/RivermaxCore/Source/RivermaxRendering/Public/RivermaxShaders.h
--- a/Engine/Plugins/VirtualProduction/Rivermax/RivermaxCore/Source/RivermaxRendering/Public/RivermaxShaders.h
+++ b/Engine/Plugins/VirtualProduction/Rivermax/RivermaxCore/Source/RivermaxRendering/Public/RivermaxShaders.h
@@ -30,9 +30,6 @@ namespace UE::RivermaxShaders
 		DECLARE_EXPORTED_GLOBAL_SHADER(FYUV8Bit422ToRGBACS, RIVERMAXRENDERING_API);
 		SHADER_USE_PARAMETER_STRUCT(FYUV8Bit422ToRGBACS, FGlobalShader);
 
-		class FSRGBToLinear : SHADER_PERMUTATION_BOOL("DO_SRGB_TO_LINEAR");
-		using FPermutationDomain = TShaderPermutationDomain<FSRGBToLinear>;
-
 		BEGIN_SHADER_PARAMETER_STRUCT(FParameters, )
 			SHADER_PARAMETER_RDG_BUFFER_SRV(StructuredBuffer<FYUV8Bit422Buffer>, InputYUV4228bitBuffer)
 			SHADER_PARAMETER(FMatrix44f, ColorTransform)
@@ -111,9 +108,6 @@ namespace UE::RivermaxShaders
 		DECLARE_EXPORTED_GLOBAL_SHADER(FYUV10Bit422ToRGBACS, RIVERMAXRENDERING_API);
 		SHADER_USE_PARAMETER_STRUCT(FYUV10Bit422ToRGBACS, FGlobalShader);
 
-		class FSRGBToLinear : SHADER_PERMUTATION_BOOL("DO_SRGB_TO_LINEAR");
-		using FPermutationDomain = TShaderPermutationDomain<FSRGBToLinear>;
-
 		BEGIN_SHADER_PARAMETER_STRUCT(FParameters, )
 			SHADER_PARAMETER_RDG_BUFFER_SRV(StructuredBuffer<FYUV10Bit422LEBuffer>, InputYCbCrBuffer)
 			SHADER_PARAMETER(FMatrix44f, ColorTransform)	
@@ -197,9 +191,6 @@ namespace UE::RivermaxShaders
 		DECLARE_EXPORTED_GLOBAL_SHADER(FRGB8BitToRGBA8CS, RIVERMAXRENDERING_API);
 		SHADER_USE_PARAMETER_STRUCT(FRGB8BitToRGBA8CS, FGlobalShader);
 
-		class FSRGBToLinear : SHADER_PERMUTATION_BOOL("DO_SRGB_TO_LINEAR");
-		using FPermutationDomain = TShaderPermutationDomain<FSRGBToLinear>;
-
 		BEGIN_SHADER_PARAMETER_STRUCT(FParameters, )
 			SHADER_PARAMETER_RDG_BUFFER_SRV(StructuredBuffer<FRGB8BitBuffer>, InputRGB8Buffer)
 			SHADER_PARAMETER(uint32, TotalElementCount)
@@ -324,9 +315,6 @@ namespace UE::RivermaxShaders
 		DECLARE_EXPORTED_GLOBAL_SHADER(FRGB10BitToRGBA10CS, RIVERMAXRENDERING_API);
 		SHADER_USE_PARAMETER_STRUCT(FRGB10BitToRGBA10CS, FGlobalShader);
 
-		class FSRGBToLinear : SHADER_PERMUTATION_BOOL("DO_SRGB_TO_LINEAR");
-		using FPermutationDomain = TShaderPermutationDomain<FSRGBToLinear>;
-
 		BEGIN_SHADER_PARAMETER_STRUCT(FParameters, )
 			SHADER_PARAMETER_RDG_BUFFER_SRV(StructuredBuffer<FRGB10BitBuffer>, InputBuffer)
 			SHADER_PARAMETER(uint32, TotalElementCount)
@@ -407,9 +395,6 @@ namespace UE::RivermaxShaders
 		DECLARE_EXPORTED_GLOBAL_SHADER(FRGB12BitToRGBA12CS, RIVERMAXRENDERING_API);
 		SHADER_USE_PARAMETER_STRUCT(FRGB12BitToRGBA12CS, FGlobalShader);
 
-		class FSRGBToLinear : SHADER_PERMUTATION_BOOL("DO_SRGB_TO_LINEAR");
-		using FPermutationDomain = TShaderPermutationDomain<FSRGBToLinear>;
-
 		BEGIN_SHADER_PARAMETER_STRUCT(FParameters, )
 			SHADER_PARAMETER_RDG_BUFFER_SRV(StructuredBuffer<FRGB12BitBuffer>, InputRGB12Buffer)
 			SHADER_PARAMETER(uint32, TotalElementCount)
@@ -488,9 +473,6 @@ namespace UE::RivermaxShaders
 		DECLARE_EXPORTED_GLOBAL_SHADER(FRGB16fBitToRGBA16fCS, RIVERMAXRENDERING_API);
 		SHADER_USE_PARAMETER_STRUCT(FRGB16fBitToRGBA16fCS, FGlobalShader);
 
-		class FSRGBToLinear : SHADER_PERMUTATION_BOOL("DO_SRGB_TO_LINEAR");
-		using FPermutationDomain = TShaderPermutationDomain<FSRGBToLinear>;
-
 		BEGIN_SHADER_PARAMETER_STRUCT(FParameters, )
 			SHADER_PARAMETER_RDG_BUFFER_SRV(StructuredBuffer<FRGB16fBuffer>, InputRGB16fBuffer)
 			SHADER_PARAMETER(uint32, TotalElementCount)
diff --git a/Engine/Plugins/VirtualProduction/Rivermax/RivermaxMedia/Source/RivermaxMedia/Private/RivermaxMediaPlayer.cpp b/Engine/Plugins/VirtualProduction/Rivermax/RivermaxMedia/Source/RivermaxMedia/Private/RivermaxMediaPlayer.cpp
--- a/Engine/Plugins/VirtualProduction/Rivermax/RivermaxMedia/Source/RivermaxMedia/Private/RivermaxMediaPlayer.cpp
+++ b/Engine/Plugins/VirtualProduction/Rivermax/RivermaxMedia/Source/RivermaxMedia/Private/RivermaxMediaPlayer.cpp
@@ -17,6 +17,7 @@
 #include "RivermaxMediaTextureSample.h"
 #include "RivermaxMediaUtils.h"
 #include "RivermaxPTPUtils.h"
+#include "RivermaxShaders.h"
 #include "RivermaxTracingUtils.h"
 #include "RivermaxTypes.h"
 #include "Stats/Stats2.h"
@@ -29,12 +30,21 @@
 
 #define LOCTEXT_NAMESPACE "FRivermaxMediaPlayer"
 
-DECLARE_CYCLE_STAT(TEXT("Rivermax MediaPlayer Request frame"), STAT_Rivermax_MediaPlayer_RequestFrame, STATGROUP_Media);
-DECLARE_CYCLE_STAT(TEXT("Rivermax MediaPlayer Process frame"), STAT_Rivermax_MediaPlayer_ProcessFrame, STATGROUP_Media);
+// Sample number with frame delay taken into account
+#define TO_EXPECTED_SAMPLE_FRAME_NUM(CURRENTFrameNum) GetFrameNumberWithAcountedLatency(FrameDelay, CURRENTFrameNum)
+
+// Sample index with frame delay taken into account
+#define TO_EXPECTED_SAMPLE_INDEX(CURRENTFrameNum) ConvertFrameNumToSampleIndex(FrameDelay, CURRENTFrameNum, kMaxNumVideoFrameBuffer)
+
+// Sample index to access sample in framelocked array
+#define FRAME_NUM_TO_INDEX(CURRENTFrameNum) (CURRENTFrameNum) % kMaxNumVideoFrameBuffer
+
+// Identifies if the player is in framelocking mode.
+#define IS_FRAMELOCKED() (EvaluationType == EMediaIOSampleEvaluationType::Timecode && bFramelock)
 
 DECLARE_GPU_STAT_NAMED(RivermaxMedia_SampleUsageFence, TEXT("RivermaxMedia_SampleUsageFence"));
 DECLARE_GPU_STAT_NAMED(Rmax_WaitForPixels, TEXT("Rmax_WaitForPixels"));
-
+DECLARE_GPU_STAT(RivermaxSource_SampleConversion);
 
 namespace UE::RivermaxMedia
 {
@@ -44,25 +54,30 @@ namespace UE::RivermaxMedia
 		TEXT("Override latency in framelock mode. 0 for 0 frame of latency and 1 for 1 frame of latency."),
 		ECVF_Default);
 
-	static TAutoConsoleVariable<int32> CVarRivermaxSampleUploadMode(
-		TEXT("Rivermax.Player.UploadMode"),
-		1,
-		TEXT("Mode 0: Upload is done on the render thread.\n"
-			"Mode 1: Upload is done in its own thread before being rendered."),
+	static TAutoConsoleVariable<float> CVarRivermaxDefaultThreadBlockTimeoutSec(
+		TEXT("Rivermax.Player.DefaultThreadBlockTimeoutSec"),
+		0.5,
+		TEXT("Overrides the wait time for the samples to be written to. In seconds. "),
 		ECVF_Default);
 
+	uint64 GetFrameNumberWithAcountedLatency(int64 InFrameDelay, uint64 FrameNumber)
+	{
+		const int64 ForcedLatency = CVarRivermaxForcedFramelockLatency.GetValueOnAnyThread();
+		return (FrameNumber - (uint64)FMath::Clamp(InFrameDelay + ForcedLatency, 0, 1));
+	}
+
+	uint8 ConvertFrameNumToSampleIndex(int64 InFrameDelay, uint64 FrameNumber, uint8 InMaxNumVideoFrameBuffer)
+	{
+		return (GetFrameNumberWithAcountedLatency(InFrameDelay, FrameNumber)) % InMaxNumVideoFrameBuffer;
+	}
+
 	/* FRivermaxVideoPlayer structors
 	 *****************************************************************************/
 
 	FRivermaxMediaPlayer::FRivermaxMediaPlayer(IMediaEventSink& InEventSink)
 		: Super(InEventSink)
-		, MaxNumVideoFrameBuffer(5)
 		, RivermaxThreadNewState(EMediaState::Closed)
-		, bIsSRGBInput(false)
-		, bUseVideo(false)
-		, SupportedSampleTypes(EMediaIOSampleType::None)
-		, bPauseRequested(false)
-		, MediaSamples(MakeUnique<FRivermaxMediaTextureSamples>())
+		, VideoTextureSamplePool(MakeUnique<FRivermaxMediaTextureSamplePool>())
 	{
 	}
 
@@ -88,21 +103,13 @@ namespace UE::RivermaxMedia
 			return false;
 		}
 
-		PlayerMode = (ERivermaxPlayerMode)Options->GetMediaOption(RivermaxMediaOption::PlayerMode, (int64)ERivermaxPlayerMode::Latest);
-		FrameTracking = {};
-
 		//Video related options
 		{
-			bIsSRGBInput = Options->GetMediaOption(RivermaxMediaOption::SRGBInput, bIsSRGBInput);
 			DesiredPixelFormat = (ERivermaxMediaSourcePixelFormat)Options->GetMediaOption(RivermaxMediaOption::PixelFormat, (int64)ERivermaxMediaSourcePixelFormat::RGB_8bit);
-			const bool bUseZeroLatency = Options->GetMediaOption(RivermaxMediaOption::ZeroLatency, true);
-			FrameLatency = bUseZeroLatency ? 0 : 1;
 			const bool bOverrideResolution = Options->GetMediaOption(RivermaxMediaOption::OverrideResolution, false);
 			bFollowsStreamResolution = !bOverrideResolution;
 		}
 
-		bUseVideo = true;
-
 		IRivermaxCoreModule* Module = FModuleManager::GetModulePtr<IRivermaxCoreModule>("RivermaxCore");
 		if (Module && ConfigureStream(Options))
 		{
@@ -114,13 +121,12 @@ namespace UE::RivermaxMedia
 		{
 			StreamResolution = StreamOptions.EnforcedResolution;
 		}
+
 		VideoTrackFormat.Dim = FIntPoint::ZeroValue;
 
-		FrameTracking.bWasFrameRequested = false;
-		FrameTracking.LastFrameRendered.Reset();
 		CurrentState = EMediaState::Preparing;
 		RivermaxThreadNewState = EMediaState::Preparing;
-		
+
 		if (InputStream == nullptr || !InputStream->Initialize(StreamOptions, *this))
 		{
 			UE_LOG(LogRivermaxMedia, Warning, TEXT("Failed to initialize Rivermax input stream."));
@@ -132,8 +138,6 @@ namespace UE::RivermaxMedia
 		// Setup our different supported channels based on source settings
 		SetupSampleChannels();
 
-		EventSink.ReceiveMediaEvent(EMediaEvent::MediaConnecting);
-
 #if WITH_EDITOR
 		if (FEngineAnalytics::IsAvailable())
 		{
@@ -162,15 +166,8 @@ namespace UE::RivermaxMedia
 			InputStream.Reset();
 		}
 
-		for(const TSharedPtr<FRivermaxSampleWrapper>& Sample : SamplePool)
-		{
-			Sample->Sample.Reset();
-			Sample->SampleConversionFence.SafeRelease();
-		}
-		SamplePool.Empty();
-
-		RivermaxThreadCurrentTextureSample.Reset();
-		SkippedFrames.Empty();
+		Samples->FlushSamples();
+		VideoTextureSamplePool.Reset();
 
 		Super::Close();
 	}
@@ -181,138 +178,99 @@ namespace UE::RivermaxMedia
 		return PlayerPluginGUID;
 	}
 
-	FString FRivermaxMediaPlayer::GetStats() const
-	{
-		FString Stats;
-
-		//todo
-
-		return Stats;
-	}
 
 #if WITH_EDITOR
 	const FSlateBrush* FRivermaxMediaPlayer::GetDisplayIcon() const
 	{
 		//todo for tdm
 		return nullptr;
-		//return IRivermaxMediaModule::Get().GetStyle()->GetBrush("RivermaxMediaIcon");
 	}
 #endif //WITH_EDITOR
 
-	bool FRivermaxMediaPlayer::OnVideoFrameRequested(const FRivermaxInputVideoFrameDescriptor& FrameInfo, FRivermaxInputVideoFrameRequest& OutVideoFrameRequest)
+	TSharedPtr<IRivermaxVideoSample> FRivermaxMediaPlayer::OnVideoFrameRequested(const FRivermaxInputVideoFrameDescriptor& FrameInfo)
 	{
-		TRACE_CPUPROFILER_EVENT_SCOPE(FRivermaxMediaPlayer::OnVideoFrameRequested)
+		TRACE_CPUPROFILER_EVENT_SCOPE_TEXT(*FString::Printf(TEXT("Rmax::FrameRequested %u"), (FrameInfo.FrameNumber)));
 
 		// If video is not playing, no need to provide samples when requested
 		if (!IsReadyToPlay())
 		{
-			return false;
+			return nullptr;
 		}
 
-		// Track first frame number that was requested in order to avoid waiting for a frame that will never come
-		if (FrameTracking.bWasFrameRequested == false)
-		{
-			FrameTracking.bWasFrameRequested = true;
-			FrameTracking.FirstFrameRequested = FrameInfo.FrameNumber;
-		}
-		else
+		if (FrameInfo.VideoBufferSize > 0)
 		{
-			if(FrameInfo.FrameNumber > (FrameTracking.LastFrameNumberRequested + 1))
+			TSharedPtr<FRivermaxMediaTextureSample> Sample;
+			if (IS_FRAMELOCKED())
 			{
-				UE_LOG(LogRivermaxMedia, Warning, TEXT("Frames were skipped. Last frame requested was '%d'. Current frame being requested is '%d'."), FrameTracking.LastFrameNumberRequested, FrameInfo.FrameNumber);
-				
-				// Skipped frames are only used in framelock mode to avoid stalling
-				if (PlayerMode == ERivermaxPlayerMode::Framelock)
+				// If input stream has somehow completely lost a frame (not even a single packet received) and we are 2 frames ahead of rendered frame.
+				if (FrameInfo.FrameNumber > LastFrameToAttemptReception + 1)
 				{
-					FScopeLock Lock(&SkippedFrameCriticalSection);
-					
-					const uint32 MinInterval = FrameTracking.LastFrameNumberRequested + 1;
-					const uint32 MaxInterval = FrameInfo.FrameNumber - 1;
-					if (MinInterval <= MaxInterval)
+					//Invalidate previous samples in case one of them is still being rendered. 
+					for (uint64 PreviousFrameNumber = LastFrameToAttemptReception; PreviousFrameNumber < FrameInfo.FrameNumber; PreviousFrameNumber++)
 					{
-						const TInterval<uint32> Interval(MinInterval, MaxInterval);
-						SkippedFrames.Add(Interval);
-					}
-					else
-					{
-						//Wrap around case
-						const TInterval<uint32> FirstInterval(MinInterval, ~0);
-						const TInterval<uint32> SecondInterval(0, MaxInterval);
-						SkippedFrames.Add(FirstInterval);
-						SkippedFrames.Add(SecondInterval);
+						TSharedPtr<FRivermaxMediaTextureSample> PreviousSample = FrameLockedSamples[FRAME_NUM_TO_INDEX(PreviousFrameNumber)];
+						if (PreviousSample->IsBeingRendered())
+						{
+							// This will stop the waiting thread from waiting for the start of the reception. 
+							PreviousSample->SetFrameNumber(PreviousFrameNumber);
+
+							// this will stop the waiting thread from copying invalid data and waiting for the copy.
+							PreviousSample->SetReceptionState(IRivermaxVideoSample::ESampleState::ReceptionError);
+
+							// this will stop the waiting thread from waiting for the sample reception.
+							PreviousSample->GetSampleReceivedEvent()->Trigger();
+						}
 					}
 				}
-			}
-		}
 
-		// Keep track of frames that are requested to detect gaps
-		FrameTracking.LastFrameNumberRequested = FrameInfo.FrameNumber;
+				Sample = FrameLockedSamples[FRAME_NUM_TO_INDEX(FrameInfo.FrameNumber)];
+				Sample->GetSampleReceivedEvent()->Reset();
 
-		if (FrameInfo.VideoBufferSize > 0)
-		{
-			uint32 NextRequestIndex = 0;
-			if(GetFrameRequestedIndex(FrameInfo, NextRequestIndex))
+			}
+			else
 			{
-				UE_LOG(LogRivermaxMedia, Verbose, TEXT("Starting to receive frame '%u' with timestamp %u at location %d"), FrameInfo.FrameNumber, FrameInfo.Timestamp, NextRequestIndex);
-				SamplePool[NextRequestIndex]->ReceptionState = ESampleReceptionState::Receiving;
-				SamplePool[NextRequestIndex]->FrameNumber = FrameInfo.FrameNumber;
-				SamplePool[NextRequestIndex]->Timestamp = FrameInfo.Timestamp;
-				RivermaxThreadCurrentTextureSample = SamplePool[NextRequestIndex];
-			
-				if (bDoesStreamSupportsGPUDirect)
-				{
-					// Hand out location where to copy received data on gpu
-					OutVideoFrameRequest.GPUBuffer = RivermaxThreadCurrentTextureSample->Sample->GetGPUBuffer()->GetRHI();
-				}
-				else
-				{
-					OutVideoFrameRequest.VideoBuffer = reinterpret_cast<uint8*>(RivermaxThreadCurrentTextureSample->Sample->RequestBuffer(FrameInfo.VideoBufferSize));
-				}
-				return true;
+				Sample = VideoTextureSamplePool->AcquireShared(false /*NoAllocation*/);
+			}
+
+			// (More of a sanity check. Shouldn't be in this state.).
+			if (!Sample.IsValid())
+			{
+				UE_LOG(LogRivermaxMedia, Warning, TEXT("Failed to provide a frame for incoming frame %u with timestamp %u"), FrameInfo.FrameNumber, FrameInfo.Timestamp);
+				return nullptr;
 			}
 
-			UE_LOG(LogRivermaxMedia, Verbose, TEXT("Failed to provide a frame for incoming frame %u with timestamp %u"), FrameInfo.FrameNumber, FrameInfo.Timestamp);
+			UE_LOG(LogRivermaxMedia, Verbose, TEXT("Starting to receive frame '%u' with timestamp %u"), FrameInfo.FrameNumber, FrameInfo.Timestamp);
+
+			Sample->SetFrameNumber(FrameInfo.FrameNumber);
+			const double MediaFrameTimeSecs = UE::RivermaxCore::ConvertRTPTimeStampToSecondsTruncated(FrameInfo.Timestamp);
+			Sample->SetTime(FTimespan::FromSeconds(MediaFrameTimeSecs));
+			Sample->SetReceptionState(IRivermaxVideoSample::ESampleState::ReadyForReception);
+			LastFrameToAttemptReception = FrameInfo.FrameNumber;
+			return Sample;
 		}
 
-		return false;
+		return nullptr;
 	}
 
 
-	void FRivermaxMediaPlayer::OnVideoFrameReceived(const FRivermaxInputVideoFrameDescriptor& FrameInfo, const FRivermaxInputVideoFrameReception& ReceivedVideoFrame)
+	void FRivermaxMediaPlayer::OnVideoFrameReceived(TSharedPtr<IRivermaxVideoSample> InReceivedVideoFrame)
 	{
-		TRACE_CPUPROFILER_EVENT_SCOPE(FRivermaxMediaPlayer::OnVideoFrameReceived)
-
-
 		if (!IsReadyToPlay())
 		{
 			return;
 		}
 
-		if (bUseVideo && ReceivedVideoFrame.VideoBuffer)
+		TSharedPtr<FRivermaxMediaTextureSample> Sample = StaticCastSharedPtr<FRivermaxMediaTextureSample>(InReceivedVideoFrame);
+		check(Sample.IsValid());
+
+		TRACE_CPUPROFILER_EVENT_SCOPE_TEXT(*FString::Printf(TEXT("Rmax::OnVideoFrameReceived %u"), (Sample->GetFrameNumber())));
+		if (!IS_FRAMELOCKED())
 		{
-			if (RivermaxThreadCurrentTextureSample.IsValid())
-			{
-				if(RivermaxThreadCurrentTextureSample->ReceptionState == ESampleReceptionState::Receiving)
-				{
-					RivermaxThreadCurrentTextureSample->ReceptionState = ESampleReceptionState::Received;
-					RivermaxThreadCurrentTextureSample->bIsReadyToRender = bDoesStreamSupportsGPUDirect;
-				}
-				else
-				{
-					UE_LOG(LogRivermaxMedia, Verbose, TEXT("Discarding received frame %u since it was deemed unusable."), FrameInfo.FrameNumber);
-				}
-			}
+			Samples->AddVideo(Sample.ToSharedRef());
 		}
-		
-		RivermaxThreadCurrentTextureSample.Reset();
-	}
-
-	void FRivermaxMediaPlayer::TickFetch(FTimespan DeltaTime, FTimespan Timecode)
-	{
-		Super::TickFetch(DeltaTime, Timecode);
-		if (InputStream && CurrentState == EMediaState::Playing)
+		else
 		{
-			ProcessFrame();
+			Sample->GetSampleReceivedEvent()->Trigger();
 		}
 	}
 
@@ -341,7 +299,7 @@ namespace UE::RivermaxMedia
 		{
 			return;
 		}
-		
+
 		// Cache current stream detection, it could change while we are applying it
 		FIntPoint CachedStreamResolution;
 		{
@@ -352,27 +310,12 @@ namespace UE::RivermaxMedia
 		if (VideoTrackFormat.Dim != CachedStreamResolution)
 		{
 			UE_LOG(LogRivermaxMedia, Log, TEXT("Player needs to apply newly detected stream resolution : %dx%d"), CachedStreamResolution.X, CachedStreamResolution.Y);
-			
-			// Reset some frame tracking info while changing resolution
-			FrameTracking.bWasFrameRequested = false;
-			FrameTracking.LastFrameRendered.Reset();
 
 			{
 				WaitForPendingTasks();
 
-				// Cleanup allocated ressources for the current resolution
-				for (const TSharedPtr<FRivermaxSampleWrapper>& Sample : SamplePool)
-				{
-					Sample->Sample.Reset();
-					Sample->SampleConversionFence.SafeRelease();
-				}
-				SamplePool.Empty();
-
-				RivermaxThreadCurrentTextureSample.Reset();
-				SkippedFrames.Empty();
-
 				AllocateBuffers(CachedStreamResolution);
-				
+
 				VideoTrackFormat.Dim = CachedStreamResolution;
 			}
 
@@ -381,37 +324,8 @@ namespace UE::RivermaxMedia
 		TickTimeManagement();
 	}
 
-
 	/* FRivermaxMediaPlayer implementation
 	 *****************************************************************************/
-	void FRivermaxMediaPlayer::ProcessFrame()
-	{
-		TRACE_CPUPROFILER_EVENT_SCOPE(RivermaxPlayerProcessFrame);
-
-		// Don't start making frame available until one is being received
-		if (FrameTracking.bWasFrameRequested == false)
-		{
-			return;
-		}
-
-		// Create a new media sample acting as a dummy container to be picked by MFW which we will fill during the late update
-		FSampleConfigurationArgs Args;
-		Args.bInIsSRGBInput = bIsSRGBInput;
-		Args.FrameRate = VideoFrameRate;
-		Args.Width = VideoTrackFormat.Dim.X;
-		Args.Height = VideoTrackFormat.Dim.Y;
-		Args.Player = StaticCastSharedPtr<FRivermaxMediaPlayer>(AsShared().ToSharedPtr());
-		Args.SampleFormat = DesiredPixelFormat;
-		Args.Time = FTimespan(GFrameCounter);
-
-		TSharedPtr<FRivermaxMediaTextureSample> EmptySample = MakeShared<FRivermaxMediaTextureSample>();
-		EmptySample->ConfigureSample(Args);
-		EmptySample->SetBuffer(CommonGPUBuffer);
-
-		// Update the current sample to be picked by MFW
-		MediaSamples->CurrentSample = EmptySample;
-	}
-
 	bool FRivermaxMediaPlayer::IsHardwareReady() const
 	{
 		return (RivermaxThreadNewState == EMediaState::Playing) || (RivermaxThreadNewState == EMediaState::Paused);
@@ -420,31 +334,38 @@ namespace UE::RivermaxMedia
 	void FRivermaxMediaPlayer::SetupSampleChannels()
 	{
 		FMediaIOSamplingSettings VideoSettings = BaseSettings;
-		VideoSettings.BufferSize = MaxNumVideoFrameBuffer;
+		VideoSettings.BufferSize = kMaxNumVideoFrameBuffer;
+
+		// TODO: Initialize Audio, Anc buffers
 		Samples->InitializeVideoBuffer(VideoSettings);
 	}
 
-	bool FRivermaxMediaPlayer::SetRate(float Rate)
+	TSharedPtr<FMediaIOCoreTextureSampleConverter> FRivermaxMediaPlayer::CreateTextureSampleConverter() const
 	{
-		if (FMath::IsNearlyEqual(Rate, 1.0f))
+		return MakeShared<FRivermaxMediaTextureSampleConverter>();
+	}
+
+	TSharedPtr<FMediaIOCoreTextureSampleBase> FRivermaxMediaPlayer::AcquireTextureSample_AnyThread() const
+	{
+
+		if ((Samples->NumVideoSamples() > 0 || IsJustInTimeRenderingEnabled()) && ProxySampleDummy.IsValid())
 		{
-			bPauseRequested = false;
-			return true;
+			FScopeLock Lock(&ProxySampleAccessCriticalSection);
+			// Create a copy of the proxy sample as Media Texture uses raw pointers for converters on Render thread and converter is set on game thread.
+			TSharedPtr<FRivermaxMediaTextureSample> SampleToReturn = MakeShared<FRivermaxMediaTextureSample>();
+			SampleToReturn->CopyConfiguration(ProxySampleDummy);
+			return SampleToReturn;
 		}
-
-		if (FMath::IsNearlyEqual(Rate, 0.0f))
+		else
 		{
-			bPauseRequested = true;
-			return true;
+			return nullptr;
 		}
-
-		return false;
 	}
 
 	void FRivermaxMediaPlayer::OnInitializationCompleted(const FRivermaxInputInitializationResult& Result)
 	{
 		RivermaxThreadNewState = Result.bHasSucceed ? EMediaState::Playing : EMediaState::Error;
-		bDoesStreamSupportsGPUDirect = Result.bIsGPUDirectSupported;
+		bStreamSupportsGPUDirect = Result.bIsGPUDirectSupported;
 	}
 
 	bool FRivermaxMediaPlayer::ConfigureStream(const IMediaOptions* Options)
@@ -482,46 +403,59 @@ namespace UE::RivermaxMedia
 
 	void FRivermaxMediaPlayer::AllocateBuffers(const FIntPoint& InResolution)
 	{
+		TRACE_CPUPROFILER_EVENT_SCOPE(FRivermaxMediaPlayer::AllocateBuffers);
 		using namespace UE::RivermaxCore;
 		using namespace UE::RivermaxMediaUtils::Private;
 
-		// Take care of the common buffers first
+		if (VideoTextureSamplePool.IsValid())
 		{
-			// Create the common texture we are going to use
-			const FSourceBufferDesc BufferDescription = GetBufferDescription(InResolution, DesiredPixelFormat);
-			FRDGBufferDesc RDGBufferDesc = FRDGBufferDesc::CreateStructuredDesc(BufferDescription.BytesPerElement, BufferDescription.NumberOfElements);
+			VideoTextureSamplePool->Reset();
+		}
 
-			// Required to share resource across different graphics API (DX, Cuda)
-			RDGBufferDesc.Usage |= EBufferUsageFlags::Shared;
+		VideoTextureSamplePool = MakeUnique<FRivermaxMediaTextureSamplePool>();
+		Samples->FlushSamples();
 
-			TWeakPtr<FRivermaxMediaPlayer> WeakPlayer = StaticCastSharedRef<FRivermaxMediaPlayer>(AsShared());
-			ENQUEUE_RENDER_COMMAND(RivermaxPlayerBufferCreation)(
-				[WeakPlayer, RDGBufferDesc](FRHICommandListImmediate& CommandList)
-				{
-					if (TSharedPtr<FRivermaxMediaPlayer> Player = WeakPlayer.Pin())
-					{
-						Player->CommonGPUBuffer = AllocatePooledBuffer(RDGBufferDesc, TEXT("RmaxInput Buffer"));
-					}
-				});
-		}
+		// We only need to store current and the number we delay by.
+		Samples->SetSampleBufferSize(FrameDelay + 1);
 
-		SamplePool.Empty();
 
+		TSharedPtr<FRivermaxMediaTextureSample> NewSample;
 		// Allocate our pool of samples where incoming ones will be written and chosen from
-		for (int32 Index = 0; Index < MaxNumVideoFrameBuffer; Index++)
-		{
-			TSharedPtr<FRivermaxSampleWrapper> NewWrapper = MakeShared<FRivermaxSampleWrapper>();
-			NewWrapper->Sample = MakeShared<FRivermaxMediaTextureSample>();
-			NewWrapper->Sample->InitializeGPUBuffer(InResolution, DesiredPixelFormat);
-			NewWrapper->SampleConversionFence = RHICreateGPUFence(*FString::Printf(TEXT("RmaxConversionDoneFence_%02d"), Index));
-			NewWrapper->ReceptionState = ESampleReceptionState::Available;
-			SamplePool.Add(MoveTemp(NewWrapper));
+		for (int32 Index = 0; Index < kMaxNumVideoFrameBuffer; Index++)
+		{
+			NewSample = VideoTextureSamplePool->AcquireShared(true /*Allocate New*/);
+			NewSample->InitializeGPUBuffer(InResolution, DesiredPixelFormat, bStreamSupportsGPUDirect);
+			NewSample->SampleConversionFence = RHICreateGPUFence(*FString::Printf(TEXT("RmaxConversionDoneFence_%02d"), Index));
+			NewSample->SetReceptionState(IRivermaxSample::ESampleState::ReadyForReception);
+			NewSample->SetInputFormat(DesiredPixelFormat);
+			if (IS_FRAMELOCKED())
+			{
+				FrameLockedSamples[Index] = NewSample;
+			}
 		}
 
+		// Create the proxy sample that is going to be used for color conversion.
+		ENQUEUE_RENDER_COMMAND(FRivermaxMediaTextureSample)(
+			[NewSample, Resolution = InResolution, this](FRHICommandListImmediate& RHICmdList)
+			{
+				FScopeLock Lock(&ProxySampleAccessCriticalSection);
+
+				ProxySampleDummy = MakeShared<FRivermaxMediaTextureSample>();
+				ProxySampleDummy->SetTexture(CreateIntermediateRenderTarget(RHICmdList, Resolution, NewSample->GetPixelFormat(), NewSample->IsOutputSrgb()));
+				UE::MediaIOCore::FColorFormatArgs ColorArgs;
+				ColorArgs.ColorSpace = bOverrideSourceColorSpace ? (UE::Color::EColorSpace)OverrideSourceColorSpace : UE::Color::EColorSpace::None;
+				ColorArgs.Encoding = bOverrideSourceEncoding ? (UE::Color::EEncoding)OverrideSourceEncoding : UE::Color::EEncoding::None;
+
+				ProxySampleDummy->SetProperties(NewSample->GetStride(), VideoTrackFormat.Dim.X, VideoTrackFormat.Dim.X, NewSample->GetFormat(), FTimespan(0), FFrameRate(), FTimecode(), ColorArgs);
+			});
+
+
 		// Allocation is done on render thread so let's make sure it's completed before pursuing
 		FRenderCommandFence RenderFence;
 		RenderFence.BeginFence();
-		RenderFence.Wait();		
+		RenderFence.Wait();
+		VideoTextureSamplePool->Tick();
+
 	}
 
 	void FRivermaxMediaPlayer::OnStreamError()
@@ -531,338 +465,277 @@ namespace UE::RivermaxMedia
 		RivermaxThreadNewState = EMediaState::Error;
 	}
 
-	bool FRivermaxMediaPlayer::LateUpdateSetupSample(FSampleConverterOperationSetup& OutConverterSetup)
+	TRefCountPtr<FRHITexture> FRivermaxMediaPlayer::CreateIntermediateRenderTarget(FRHICommandListImmediate& RHICmdList, const FIntPoint& InDim, EPixelFormat InPixelFormat, bool bInSRGB)
+	{
+		TRefCountPtr<FRHITexture> TextureToReturn;
+		// create output render target if necessary
+		ETextureCreateFlags OutputCreateFlags = TexCreate_Dynamic | (bInSRGB ? TexCreate_SRGB : TexCreate_None) | TexCreate_UAV;
+		OutputCreateFlags |= TexCreate_UAV;
+		OutputCreateFlags |= ETextureCreateFlags::RenderTargetable;
+
+
+		const static FLazyName ClassName(TEXT("FRivermaxMediaTextureSample"));
+		const FRHITextureCreateDesc Desc =
+			FRHITextureCreateDesc::Create2D(TEXT("FRivermaxMediaTextureOutput"))
+			.SetExtent(InDim)
+			.SetFormat(InPixelFormat)
+			//.SetNumMips(1)
+			.SetFlags(OutputCreateFlags | ETextureCreateFlags::ShaderResource)
+			.SetInitialState(ERHIAccess::SRVMask)
+			.SetClearValue(FClearValueBinding(FLinearColor::Red))
+			.SetClassName(ClassName)
+			.SetOwnerName(*GetMediaName().ToString());
+
+		TextureToReturn = RHICmdList.CreateTexture(Desc);
+
+		TextureToReturn->SetName(TEXT("RivermaxMediaTexture"));
+		TextureToReturn->SetOwnerName(*GetMediaName().ToString());
+		return MoveTemp(TextureToReturn);
+	}
+
+	bool FRivermaxMediaPlayer::JustInTimeSampleRender_RenderThread(FRHICommandListImmediate& RHICmdList, FTextureRHIRef& InDestinationTexture, TSharedPtr<FMediaIOCoreTextureSampleBase>& InJITRProxySample)
 	{
+		// Player renders into an intermediate render target that is then used to be converted into the right color encoding if needed.
+		InDestinationTexture = ProxySampleDummy->GetTexture();
 		TRACE_CPUPROFILER_EVENT_SCOPE(RivermaxPlayerLateUpdate);
+		TRACE_CPUPROFILER_EVENT_SCOPE_TEXT(*FString::Printf(TEXT("Rmax::StartingRender %u"), (GFrameCounterRenderThread)));
 
-		// We only allow this function to run once per frame.
-		if (LastFrameNumberThatUpdatedJustInTime == GFrameCounterRenderThread)
-		{
-			UE_LOG(LogRivermaxMedia, Verbose, TEXT("LateUpdate called more than once in GFrameCounterRenderThread %llu"), GFrameCounterRenderThread);
-			return false;
-		}
-		LastFrameNumberThatUpdatedJustInTime = GFrameCounterRenderThread;
+		// Sanity check to make sure that we are not rendering this sample twice per frame.
+		check(LastEngineRTFrameThatUpdatedJustInTime != GFrameCounterRenderThread);
+		LastEngineRTFrameThatUpdatedJustInTime = GFrameCounterRenderThread;
 
-		FFrameExpectation NextFrameExpectations;
-		const bool bShouldRender = GetNextExpectedFrameInfo(NextFrameExpectations);
+		FFrameInfo FrameInformation;
+		FrameInformation.RequestedTimecode = InJITRProxySample->GetTimecode().Get(FTimecode());
+		FrameInformation.SampleTimespan = InJITRProxySample->GetTime().Time;
+		FrameInformation.EvaluationOffset = InJITRProxySample->GetEvaluationOffsetInSeconds();
+		FrameInformation.FrameNumber = GFrameCounterRenderThread;
 
-		if (!bShouldRender)
+		TSharedPtr<FRivermaxMediaTextureSample> SamplePtr = StaticCastSharedPtr<FRivermaxMediaTextureSample>(PickSampleToRender_RenderThread(FrameInformation));
+		if (!SamplePtr.IsValid())
 		{
-			UE_LOG(LogRivermaxMedia, VeryVerbose, TEXT("Skipping render for frame %llu."), GFrameCounterRenderThread);
+			UE_LOG(LogRivermaxMedia, Verbose, TEXT("Couldn't find a sample to render for frame %u."), FrameInformation.FrameNumber);
 			return false;
 		}
 
-		if (NextFrameExpectations.FrameNumber < FrameTracking.FirstFrameRequested)
+		if (!InDestinationTexture.IsValid())
 		{
-			UE_LOG(LogRivermaxMedia, VeryVerbose, TEXT("Skipping render for frame number %llu. Expecting frame %u but first frame received is greater, %u."), GFrameCounterRenderThread, NextFrameExpectations.FrameNumber, FrameTracking.FirstFrameRequested);
+			UE_LOG(LogRivermaxMedia, Warning, TEXT("Couldn't find texture to render into for sample %u."), SamplePtr->GetFrameNumber());
 			return false;
 		}
 
-		// Always clean up frames from the past. They will never be used
-		for (const TSharedPtr<FRivermaxSampleWrapper>& Frame : SamplePool)
-		{
-			if (Frame->ReceptionState != ESampleReceptionState::Available)
-			{
-				if (Frame->FrameNumber < NextFrameExpectations.FrameNumber)
-				{
-					Frame->ReceptionState = ESampleReceptionState::Available;
-
-					const FString LastRender = FrameTracking.LastFrameRendered.IsSet() ? FString::Printf(TEXT("%u"), FrameTracking.LastFrameRendered.GetValue()) : FString(TEXT("None"));
-					UE_LOG(LogRivermaxMedia, Verbose, TEXT("Making frame %u as available since it will never be processed. Last render = %s, next render = %u"), Frame->FrameNumber, *LastRender, NextFrameExpectations.FrameNumber);
-				}
-			}
-		}
-		
-		FrameTracking.LastFrameRendered = NextFrameExpectations.FrameNumber;
-
 		// Verify if the frame we will use for rendering is still being rendered for the previous one.
-		if (SamplePool[NextFrameExpectations.FrameIndex]->bIsPendingRendering)
+		if (!SamplePtr->TryLockForRendering())
 		{
-			UE_LOG(LogRivermaxMedia, Verbose, TEXT("Frame %u was still rendering when we expected to reuse its location. Waiting for it to complete."), NextFrameExpectations.FrameIndex);
-
-			TRACE_CPUPROFILER_EVENT_SCOPE(RmaxMediaWaitForFrameToBeAvailable);
-
-			// Frame should be rendered at some point so a timeout isn't required but we add one as a last resort to avoid lockups.
-			constexpr double TimeoutSeconds = 2.0;
-			const double StartTimeSeconds = FPlatformTime::Seconds();
-			while (SamplePool[NextFrameExpectations.FrameIndex]->bIsPendingRendering)
+			if (IS_FRAMELOCKED())
 			{
-				FPlatformProcess::SleepNoStats(SleepTimeSeconds);
-
-				if ((FPlatformTime::Seconds() - StartTimeSeconds) > TimeoutSeconds)
-				{
-					break;
-				}
-			}
-		}
-
-		// Mark this sample as pending since it's now going on the rendering path
-		SamplePool[NextFrameExpectations.FrameIndex]->bIsPendingRendering = true;
-
-		{
-			const int32 SampleUploadMode = CVarRivermaxSampleUploadMode.GetValueOnRenderThread();
-			if (SampleUploadMode == 0)
-			{
-				SampleUploadSetupRenderThreadMode(NextFrameExpectations, OutConverterSetup);
-				
+				UE_LOG(LogRivermaxMedia, Warning, TEXT("Framelocked sample %u was still rendering when we expected to reuse its location."), SamplePtr->GetFrameNumber());
+				return false;
 			}
-			else 
+			else
 			{
-				SampleUploadSetupTaskThreadMode(NextFrameExpectations, OutConverterSetup);
+				UE_LOG(LogRivermaxMedia, Verbose, TEXT("Sample %u was either already rendered or is already being rendered."), SamplePtr->GetFrameNumber());
+				return false;
 			}
-
-			// Setup post sample usage pass 
-			OutConverterSetup.PostConvertFunc = [NextFrameExpectations, this] (FRDGBuilder& GraphBuilder)
-			{
-				PostSampleUsage(GraphBuilder, NextFrameExpectations);
-			};
 		}
 
-		return true;
-	}
+		FSampleConverterOperationSetup ConverterSetup;
+		SampleUploadSetupTaskThreadMode(SamplePtr, ConverterSetup);
 
-	bool FRivermaxMediaPlayer::GetNextExpectedFrameInfo(FFrameExpectation& OutExpectation)
-	{
-		switch (PlayerMode)
+		// If no input data was provided, no need to render
+		if (ConverterSetup.GetGPUBufferFunc == nullptr && ConverterSetup.GetSystemBufferFunc == nullptr)
 		{
-			case ERivermaxPlayerMode::Framelock:
-			{
-				return GetNextExpectedFrameInfoForFramelock(OutExpectation);
-			}
-			case ERivermaxPlayerMode::Latest:
-			{
-				return GetNextExpectedFrameInfoForLatest(OutExpectation);
-			}
-			default:
-			{
-				checkNoEntry()
-			}
+			ensureMsgf(false, TEXT("Rivermax player late update succeeded but didn't provide any source data."));
+			return false;
 		}
 
-		return false;
-	}
-
-	bool FRivermaxMediaPlayer::GetNextExpectedFrameInfoForFramelock(FFrameExpectation& OutExpectation)
-	{
-		uint32 CurrentLatency = FrameLatency;
-		const uint32 ForcedLatency = CVarRivermaxForcedFramelockLatency.GetValueOnRenderThread();
-		if (ForcedLatency == 0 || ForcedLatency == 1)
+		FRDGBuilder GraphBuilder((FRHICommandListExecutor::GetImmediateCommandList()));
+		if (ConverterSetup.PreConvertFunc)
 		{
-			CurrentLatency = ForcedLatency;
+			ConverterSetup.PreConvertFunc(GraphBuilder);
 		}
 
-		OutExpectation.FrameNumber = GFrameCounterRenderThread - CurrentLatency;
-		OutExpectation.FrameIndex = OutExpectation.FrameNumber % MaxNumVideoFrameBuffer;
-		return true;
-	}
+		using namespace UE::RivermaxShaders;
+		using namespace UE::RivermaxMediaUtils::Private;
 
-	bool FRivermaxMediaPlayer::GetNextExpectedFrameInfoForLatest(FFrameExpectation& OutExpectation)
-	{
+		const FSourceBufferDesc SourceBufferDesc = GetBufferDescription(VideoTrackFormat.Dim, SamplePtr->GetInputFormat());
 		{
-			// Look for the latest frame already arrived
-			int32 HighestTimestampIndex = INDEX_NONE;
-			uint32 HighestTimestamp = 0;
+			FRDGBufferRef InputBuffer;
 
-			int32 LowestTimestampIndex = INDEX_NONE;
-			uint32 LowestTimestamp = ~0;
+			RDG_GPU_STAT_SCOPE(GraphBuilder, RivermaxSource_SampleConversion);
+			SCOPED_DRAW_EVENT(GraphBuilder.RHICmdList, Rivermax_SampleConverter);
 
-			// Use timestamps to pick the latest one. 
-			// If we have our frame rate mismatching 
-			for (int32 Index = 0; Index < SamplePool.Num(); ++Index)
+			FRDGTextureRef OutputResource = GraphBuilder.RegisterExternalTexture(CreateRenderTarget(InDestinationTexture, TEXT("RivermaxMediaTextureOutputResource")));
+
+			// If we have a valid GPUBuffer, i.e GPUDirect is involved, use that one. Otherwise, take the system buffer and upload it in a new structured buffer.
+			if (ConverterSetup.GetGPUBufferFunc)
 			{
-				const TSharedPtr<FRivermaxSampleWrapper>& Frame = SamplePool[Index];
-				if (Frame->ReceptionState == ESampleReceptionState::Received)
-				{
-					if (HighestTimestamp < Frame->Timestamp)
+				InputBuffer = GraphBuilder.RegisterExternalBuffer(ConverterSetup.GetGPUBufferFunc(), TEXT("RMaxGPUBuffer"));
+			}
+			else if (ConverterSetup.GetSystemBufferFunc)
+			{
+				TRACE_CPUPROFILER_EVENT_SCOPE(RivermaxSampleConverter::CreateStructuredBuffer);
+				FRDGBufferNumElementsCallback NumElementCallback = [NumElements = SourceBufferDesc.NumberOfElements]()
 					{
-						HighestTimestamp = Frame->Timestamp;
-						HighestTimestampIndex = Index;
-					}
+						return NumElements;
+					};
 
-					if (LowestTimestamp > Frame->Timestamp)
+				FRDGBufferInitialDataCallback GetDataCallback = ConverterSetup.GetSystemBufferFunc;
+				FRDGBufferInitialDataSizeCallback TotalSizeCallback = [TotalSize = SourceBufferDesc.BytesPerElement * SourceBufferDesc.NumberOfElements]()
 					{
-						LowestTimestamp = Frame->Timestamp;
-						LowestTimestampIndex = Index;
-					}
-				}
-			}
+						return TotalSize;
+					};
 
-			// Early exit if no frame was found
-			if (HighestTimestampIndex < 0)
+				InputBuffer = CreateStructuredBuffer(GraphBuilder, TEXT("RivermaxInputBuffer"), SourceBufferDesc.BytesPerElement, MoveTemp(NumElementCallback), MoveTemp(GetDataCallback), MoveTemp(TotalSizeCallback));
+			}
+			else
 			{
 				return false;
 			}
 
-			check(LowestTimestampIndex >= 0);
+			const FIntPoint ProcessedOutputDimension = { (int32)SourceBufferDesc.NumberOfElements, 1 };
+			const FIntVector GroupCount = FComputeShaderUtils::GetGroupCount(SourceBufferDesc.NumberOfElements, 64);
+			FGlobalShaderMap* GlobalShaderMap = GetGlobalShaderMap(GMaxRHIFeatureLevel);
 
-			// No point in rendering the same frame as last time
-			if (FrameTracking.LastFrameExpectation.FrameNumber == SamplePool[HighestTimestampIndex]->FrameNumber)
+			//Configure shader and add conversion pass based on desired pixel format
+			switch (SamplePtr->GetInputFormat())
 			{
-				return false;
-			}
-			else if (FrameTracking.LastFrameExpectation.FrameNumber > SamplePool[HighestTimestampIndex]->FrameNumber)
+			case ERivermaxMediaSourcePixelFormat::YUV422_8bit:
 			{
-				// We have picked in the past a frame number higher than the highest in the buffers.
-				// This is unexpected, unless frame number has wrapped around.
-				UE_LOG(LogRivermaxMedia, Warning, TEXT("In latest mode, the sender's frame count seems to have reset from %u to %u."), FrameTracking.LastFrameExpectation.FrameNumber, SamplePool[HighestTimestampIndex]->FrameNumber);
+				TRACE_CPUPROFILER_EVENT_SCOPE(RivermaxSampleConverter::YUV8ShaderSetup);
+
+				const FMatrix YUVToRGBMatrix = SamplePtr->GetYUVToRGBMatrix();
+				const FVector YUVOffset(MediaShaders::YUVOffset8bits);
+				TShaderMapRef<FYUV8Bit422ToRGBACS> ComputeShader(GlobalShaderMap);
+				FYUV8Bit422ToRGBACS::FParameters* Parameters = ComputeShader->AllocateAndSetParameters(GraphBuilder, InputBuffer, OutputResource, YUVToRGBMatrix, YUVOffset, ProcessedOutputDimension.X, ProcessedOutputDimension.Y);
+
+				FComputeShaderUtils::AddPass(
+					GraphBuilder
+					, RDG_EVENT_NAME("YUV8Bit422ToRGBA")
+					, ComputeShader
+					, Parameters
+					, GroupCount);
+				break;
 			}
-
-			// Fill what we are expecting to use to render
-			OutExpectation.FrameNumber = SamplePool[HighestTimestampIndex]->FrameNumber;
-			OutExpectation.FrameIndex = HighestTimestampIndex;
-
-			// Make the ones skipped available for reception again
-			TRACE_CPUPROFILER_EVENT_SCOPE_TEXT(*FRivermaxTracingUtils::RmaxInSelectedFrameTraceEvents[OutExpectation.FrameNumber % 10]);
-
-			for (const TSharedPtr<FRivermaxSampleWrapper>& Frame : SamplePool)
+			case ERivermaxMediaSourcePixelFormat::YUV422_10bit:
 			{
-				if (Frame == SamplePool[OutExpectation.FrameIndex])
-				{
-					continue;
-				}
-
-				if (Frame->ReceptionState == ESampleReceptionState::Received)
-				{
-					UE_LOG(LogRivermaxMedia, Verbose, TEXT("Skipping frame %u since a better one was chosen, %u"), Frame->FrameNumber, OutExpectation.FrameNumber);
-					Frame->ReceptionState = ESampleReceptionState::Available;
-				}
+				TRACE_CPUPROFILER_EVENT_SCOPE(RivermaxSampleConverter::YUV10ShaderSetup);
+
+				const FMatrix YUVToRGBMatrix = SamplePtr->GetYUVToRGBMatrix();
+				const FVector YUVOffset(MediaShaders::YUVOffset10bits);
+				TShaderMapRef<FYUV10Bit422ToRGBACS> ComputeShader(GlobalShaderMap);
+				FYUV10Bit422ToRGBACS::FParameters* Parameters = ComputeShader->AllocateAndSetParameters(GraphBuilder, InputBuffer, OutputResource, YUVToRGBMatrix, YUVOffset, ProcessedOutputDimension.X, ProcessedOutputDimension.Y);
+
+				FComputeShaderUtils::AddPass(
+					GraphBuilder
+					, RDG_EVENT_NAME("YUV10Bit422ToRGBA")
+					, ComputeShader
+					, Parameters
+					, GroupCount);
+				break;
 			}
+			case ERivermaxMediaSourcePixelFormat::RGB_8bit:
+			{
+				TRACE_CPUPROFILER_EVENT_SCOPE(RivermaxSampleConverter::RGB8ShaderSetup);
 
-			FrameTracking.LastFrameExpectation = OutExpectation;
-		}
+				TShaderMapRef<FRGB8BitToRGBA8CS> ComputeShader(GlobalShaderMap);
+				FRGB8BitToRGBA8CS::FParameters* Parameters = ComputeShader->AllocateAndSetParameters(GraphBuilder, InputBuffer, OutputResource, ProcessedOutputDimension.X, ProcessedOutputDimension.Y);
 
-		return true;
-	}
+				FComputeShaderUtils::AddPass(
+					GraphBuilder
+					, RDG_EVENT_NAME("RGB8BitToRGBA8")
+					, ComputeShader
+					, Parameters
+					, GroupCount);
+				break;
+			}
+			case ERivermaxMediaSourcePixelFormat::RGB_10bit:
+			{
+				TRACE_CPUPROFILER_EVENT_SCOPE(RivermaxSampleConverter::RGB10ShaderSetup);
 
-	bool FRivermaxMediaPlayer::GetFrameRequestedIndex(const FRivermaxInputVideoFrameDescriptor& FrameInfo, uint32& OutExpectedIndex)
-	{
-		switch (PlayerMode)
-		{
-		case ERivermaxPlayerMode::Framelock:
-		{
-			return GetFrameRequestedIndexForFramelock(FrameInfo, OutExpectedIndex);
-		}
-		case ERivermaxPlayerMode::Latest:
-		{
-			return GetFrameRequestedIndexForLatest(FrameInfo, OutExpectedIndex);
-		}
-		default:
-		{
-			checkNoEntry()
-		}
-		}
+				TShaderMapRef<FRGB10BitToRGBA10CS> ComputeShader(GlobalShaderMap);
+				FRGB10BitToRGBA10CS::FParameters* Parameters = ComputeShader->AllocateAndSetParameters(GraphBuilder, InputBuffer, OutputResource, ProcessedOutputDimension.X, ProcessedOutputDimension.Y);
 
-		return false;
-	}
+				FComputeShaderUtils::AddPass(
+					GraphBuilder
+					, RDG_EVENT_NAME("RGB10BitToRGBA")
+					, ComputeShader
+					, Parameters
+					, GroupCount);
+				break;
+			}
+			case ERivermaxMediaSourcePixelFormat::RGB_12bit:
+			{
+				TRACE_CPUPROFILER_EVENT_SCOPE(RivermaxSampleConverter::RGB12ShaderSetup);
 
-	bool FRivermaxMediaPlayer::GetFrameRequestedIndexForFramelock(const FRivermaxInputVideoFrameDescriptor& FrameInfo, uint32& OutExpectedIndex)
-	{
-		// We always store incoming frame number in its respective bucket
-		OutExpectedIndex = FrameInfo.FrameNumber % MaxNumVideoFrameBuffer;
+				TShaderMapRef<FRGB12BitToRGBA12CS> ComputeShader(GlobalShaderMap);
+				FRGB12BitToRGBA12CS::FParameters* Parameters = ComputeShader->AllocateAndSetParameters(GraphBuilder, InputBuffer, OutputResource, ProcessedOutputDimension.X, ProcessedOutputDimension.Y);
 
-		TSharedPtr<FRivermaxSampleWrapper> ExpectedFrame = SamplePool[OutExpectedIndex];
+				FComputeShaderUtils::AddPass(
+					GraphBuilder
+					, RDG_EVENT_NAME("RGB12BitToRGBA")
+					, ComputeShader
+					, Parameters
+					, GroupCount);
+				break;
+			}
+			case ERivermaxMediaSourcePixelFormat::RGB_16bit_Float:
+			{
+				TRACE_CPUPROFILER_EVENT_SCOPE(RivermaxSampleConverter::RGB16FloatShaderSetup);
 
-		// If the frame is available, it hasn't been picked up for render yet or it's done rendering
-		if (ExpectedFrame->ReceptionState == ESampleReceptionState::Available)
-		{
-			return true;
-		}
+				TShaderMapRef<FRGB16fBitToRGBA16fCS> ComputeShader(GlobalShaderMap);
+				FRGB16fBitToRGBA16fCS::FParameters* Parameters = ComputeShader->AllocateAndSetParameters(GraphBuilder, InputBuffer, OutputResource, ProcessedOutputDimension.X, ProcessedOutputDimension.Y);
 
-		// If the frame is marked as received, it might not have been rendered yet 
-		else if (ExpectedFrame->ReceptionState == ESampleReceptionState::Received)
-		{
-			// Odd case where an incoming frame has the same frame number as before. In this case, we start over reception
-			if (ExpectedFrame->FrameNumber == FrameInfo.FrameNumber)
-			{
-				UE_LOG(LogRivermaxMedia, Warning, TEXT("Expected to receive frame while a matching frame number %u had not yet been rendered."), ExpectedFrame->FrameNumber);
+				FComputeShaderUtils::AddPass(
+					GraphBuilder
+					, RDG_EVENT_NAME("RGB16fBitToRGBA")
+					, ComputeShader
+					, Parameters
+					, GroupCount);
+				break;
 			}
-			else
+			default:
 			{
-				// Otherwise, it means we have overran the render so we need to wait for it to become available
-				UE_LOG(LogRivermaxMedia, Warning, TEXT("Expected to receive frame %u but location is busy with frame not rendered yet. Waiting"), ExpectedFrame->FrameNumber);
-				
-				const double StartTimeSeconds = FPlatformTime::Seconds();
-				constexpr double TimeoutSeconds = 0.5;
-				while (ExpectedFrame->ReceptionState != ESampleReceptionState::Available)
-				{
-					FPlatformProcess::SleepNoStats(SleepTimeSeconds);
-					if ((FPlatformTime::Seconds() - StartTimeSeconds) > TimeoutSeconds)
-					{
-						UE_LOG(LogRivermaxMedia, Error, TEXT("Timed out waiting for frame %u to be rendered to receive frame %u into."), ExpectedFrame->FrameNumber, FrameInfo.FrameNumber);
-					
-						// Mark this frame as skipped since we're stealing it
-						FScopeLock Lock(&SkippedFrameCriticalSection);
-						const TInterval<uint32> Interval(ExpectedFrame->FrameNumber, ExpectedFrame->FrameNumber);
-						SkippedFrames.Add(Interval);
-
-						break;
-					}
-				}
+				ensureMsgf(false, TEXT("Unhandled pixel format (%d) given to Rivermax MediaSample converter"), SamplePtr->GetInputFormat());
+				return false;
+			}
 			}
-			
-			return true;
 		}
 
-		return false;
-	}
-
-	bool FRivermaxMediaPlayer::GetFrameRequestedIndexForLatest(const FRivermaxInputVideoFrameDescriptor& FrameInfo, uint32& OutExpectedIndex)
-	{
-		//If we are free running, go forward with the write index. If we receive too fast, we'll skip some.
-		const uint32 NextIndex = ((FrameTracking.LastFrameRequestedIndex + 1) % MaxNumVideoFrameBuffer);
-		if (SamplePool[NextIndex]->ReceptionState == ESampleReceptionState::Available)
-		{
-			OutExpectedIndex = NextIndex;
-			FrameTracking.LastFrameRequestedIndex = NextIndex;
-			return true;
-		}
-		else
+		if (ConverterSetup.PostConvertFunc)
 		{
-			return false;
+			ConverterSetup.PostConvertFunc(GraphBuilder);
 		}
-	}
 
-	IMediaSamples& FRivermaxMediaPlayer::GetSamples()
-	{
-		return *MediaSamples.Get();
+		GraphBuilder.Execute();
+		return true;
 	}
 
-	void FRivermaxMediaPlayer::TickTimeManagement()
+	TSharedPtr<FMediaIOCoreTextureSampleBase> FRivermaxMediaPlayer::PickSampleToRenderFramelocked_RenderThread
+	(const FFrameInfo& InFrameInformation)
 	{
-		// When other means of alignment will be required, we will need to update this
-		// For example, aligning streams using timecode won't work with this.
-		CurrentTime = FTimespan(GFrameCounter);
-	}
+		TSharedPtr<FRivermaxMediaTextureSample> Sample = FrameLockedSamples[TO_EXPECTED_SAMPLE_INDEX(InFrameInformation.FrameNumber)];
 
-	bool FRivermaxMediaPlayer::GetPlayerFeatureFlag(EFeatureFlag flag) const
-	{
-		//switch (flag)
-		//{
-		//case EFeatureFlag::UsePlaybackTimingV2:
-		//	return true;
-		//default:
-		//	break;
-		//}
-
-		return IMediaPlayer::GetPlayerFeatureFlag(flag);
+		return Sample;
 	}
 
-	void FRivermaxMediaPlayer::PostSampleUsage(FRDGBuilder& GraphBuilder, const FFrameExpectation& FrameExpectation)
+	void FRivermaxMediaPlayer::PostSampleUsage(FRDGBuilder& GraphBuilder, TSharedPtr<FRivermaxMediaTextureSample> Sample)
 	{
 		GraphBuilder.AddPass(RDG_EVENT_NAME("RivermaxPostSampleUsage"),
-			ERDGPassFlags::NeverCull, 
-			[FrameExpectation, this](FRHICommandList& RHICmdList)
+			ERDGPassFlags::NeverCull,
+			[SamplePtr = Sample, this](FRHICommandList& RHICmdList)
 			{
 				RHI_BREADCRUMB_EVENT_STAT(RHICmdList, RivermaxMedia_SampleUsageFence, "RivermaxMedia_SampleUsageFence");
 				SCOPED_GPU_STAT(RHICmdList, RivermaxMedia_SampleUsageFence);
+				TRACE_CPUPROFILER_EVENT_SCOPE_TEXT(*FString::Printf(TEXT("Rmax::ReleasingSample %u"), (SamplePtr->GetFrameNumber())));
 
 				// Write a fence in the post sample usage pass to be able to know when we can reuse it
-				RHICmdList.WriteGPUFence(SamplePool[FrameExpectation.FrameIndex]->SampleConversionFence);
+				RHICmdList.WriteGPUFence(SamplePtr->SampleConversionFence);
 
 				++TasksInFlight;
 
+				// This async task is waiting for the GPU to be finished with Sample's resources and then release them.
 				UE::Tasks::Launch(UE_SOURCE_LOCATION,
-					[FrameExpectation, this]()
+					[SamplePtr, this]()
 					{
 						ON_SCOPE_EXIT
 						{
@@ -872,8 +745,8 @@ namespace UE::RivermaxMedia
 						TRACE_CPUPROFILER_EVENT_SCOPE(RmaxWaitForShader);
 						do
 						{
-							const bool bHasValidFence = SamplePool[FrameExpectation.FrameIndex]->SampleConversionFence.IsValid();
-							const bool bHasFenceCompleted = bHasValidFence ? SamplePool[FrameExpectation.FrameIndex]->SampleConversionFence->Poll() : false;
+							const bool bHasValidFence = SamplePtr->SampleConversionFence.IsValid();
+							const bool bHasFenceCompleted = bHasValidFence ? SamplePtr->SampleConversionFence->Poll() : false;
 							if (bHasValidFence == false || bHasFenceCompleted)
 							{
 								break;
@@ -883,70 +756,29 @@ namespace UE::RivermaxMedia
 
 						} while (true);
 
-
-						UE_LOG(LogRivermaxMedia, Verbose, TEXT("Finished rendering frame %u with expectations being frame %u."), SamplePool[FrameExpectation.FrameIndex]->FrameNumber, FrameExpectation.FrameNumber);
-
-						// We clear the fence and signal that it can be re-used.
-						SamplePool[FrameExpectation.FrameIndex]->SampleConversionFence->Clear();
-						SamplePool[FrameExpectation.FrameIndex]->ReceptionState = ESampleReceptionState::Available;
-						SamplePool[FrameExpectation.FrameIndex]->bIsReadyToRender = false;
-						SamplePool[FrameExpectation.FrameIndex]->bIsPendingRendering = false;
-						SamplePool[FrameExpectation.FrameIndex]->LockedMemory = nullptr;
-						TryClearSkippedInterval(FrameExpectation.FrameNumber);
+						// We clear the sample states, mark this sample as ready for reuse and that rendering is completed.
+						SamplePtr->ShutdownPoolable();
 					});
 			});
 	}
 
-	void FRivermaxMediaPlayer::TryClearSkippedInterval(uint32 LastFrameRendered)
+	void FRivermaxMediaPlayer::OnVideoFrameReceptionError(TSharedPtr<IRivermaxVideoSample> InVideoFrameSample)
 	{
-		FScopeLock Lock(&SkippedFrameCriticalSection);
-
-		for (auto Iter = SkippedFrames.CreateIterator(); Iter; ++Iter)
+		TSharedPtr<FRivermaxMediaTextureSample> Sample = StaticCastSharedPtr<FRivermaxMediaTextureSample>(InVideoFrameSample);
+		if (!Sample.IsValid())
 		{
-			// If we just rendered a frame matching the max boundary of an interval, we are past it from now on
-			const TInterval<uint32>& Interval = *Iter;
-			if (LastFrameRendered == Interval.Max)
-			{
-				Iter.RemoveCurrent();
-				break;
-			}
+			return;
 		}
-	}
 
-	bool FRivermaxMediaPlayer::IsFrameSkipped(uint32 FrameNumber) const
-	{
-		FScopeLock Lock(&SkippedFrameCriticalSection);
-
-		for (auto Iter = SkippedFrames.CreateConstIterator(); Iter; ++Iter)
+		UE_LOG(LogRivermaxMedia, Warning, TEXT("Issue receiving frame number %u."), Sample->GetFrameNumber());
+		if (Sample->IsBeingRendered())
 		{
-			// If we just rendered a frame matching the max boundary of an interval, we are passed it from now on
-			const TInterval<uint32>& Interval = *Iter;
-			if (Interval.Contains(FrameNumber))
-			{
-				return true;
-			}
+			Sample->SetReceptionState(FRivermaxMediaTextureSample::ESampleState::ReceptionError);
+			Sample->GetSampleReceivedEvent()->Trigger();
 		}
-
-		return false;
-	}
-
-	void FRivermaxMediaPlayer::OnVideoFrameReceptionError(const FRivermaxInputVideoFrameDescriptor& FrameInfo)
-	{
-		// In the case of an error, stamp back the frame in receiving state to available and mark it as skipped in case we are waiting on it
-		for (const TSharedPtr<FRivermaxSampleWrapper>& Sample : SamplePool)
+		else
 		{
-			// There can only be one frame receiving at the time
-			if (Sample->ReceptionState == ESampleReceptionState::Receiving)
-			{
-				UE_LOG(LogRivermaxMedia, Warning, TEXT("Error occured while receiving frame %u with timestamp %u."), Sample->FrameNumber, Sample->Timestamp);
-				Sample->ReceptionState = ESampleReceptionState::Available;
-
-				FScopeLock Lock(&SkippedFrameCriticalSection);
-				const TInterval<uint32> Interval(Sample->FrameNumber, Sample->FrameNumber);
-				SkippedFrames.Add(Interval);
-
-				return;
-			}
+			Sample->SetReceptionState(FRivermaxMediaTextureSample::ESampleState::ReadyForReception);
 		}
 	}
 
@@ -956,9 +788,9 @@ namespace UE::RivermaxMedia
 		const FIntPoint NewResolution = { (int32)NewFormatInfo.Width, (int32)NewFormatInfo.Height };
 		bool bNeedReinitializing = (NewFormatInfo.PixelFormat != StreamOptions.PixelFormat);
 		bNeedReinitializing |= (NewFormatInfo.Width != VideoTrackFormat.Dim.X || NewFormatInfo.Height != VideoTrackFormat.Dim.Y);
-		
+
 		UE_LOG(LogRivermaxMedia, Log, TEXT("New video format detected: %dx%d with pixel format '%s'"), NewResolution.X, NewResolution.Y, *UEnum::GetValueAsString(NewFormat));
-		
+
 		if (bNeedReinitializing && bFollowsStreamResolution)
 		{
 			FScopeLock Lock(&StreamResolutionCriticalSection);
@@ -996,177 +828,168 @@ namespace UE::RivermaxMedia
 		}
 	}
 
-	void FRivermaxMediaPlayer::WaitForSample(const FFrameExpectation& FrameExpectation, FWaitConditionFunc WaitConditionFunction, bool bCanTimeout)
+	bool FRivermaxMediaPlayer::WaitForSample(TSharedPtr<FRivermaxMediaTextureSample> Sample, FWaitConditionFunc WaitConditionFunction, bool bCanTimeout)
 	{
 		const double StartTimeSeconds = FPlatformTime::Seconds();
-		constexpr double TimeoutSeconds = 0.5;
+		const double TimeoutSeconds = CVarRivermaxDefaultThreadBlockTimeoutSec.GetValueOnAnyThread();
 
 		while (true)
 		{
+			if (!IsReadyToPlay())
 			{
-				if (!IsReadyToPlay())
-				{
-					break;
-				}
-
-				if (FrameExpectation.FrameIndex == INDEX_NONE)
-				{
-					break;
-				}
-
-				if (IsFrameSkipped(FrameExpectation.FrameNumber))
-				{
-					UE_LOG(LogRivermaxMedia, Verbose, TEXT("Stopped waiting for frame %u as it was marked as skipped."), FrameExpectation.FrameNumber);
-					break;
-				}
+				return false;
+			}
 
-				// Our goal here is to wait until the expected frame is available to be used (received) unless there is a timeout
-				const TSharedPtr<FRivermaxSampleWrapper>& Frame = SamplePool[FrameExpectation.FrameIndex];
-				if (WaitConditionFunction(Frame))
-				{
-					if (Frame->FrameNumber != FrameExpectation.FrameNumber)
-					{
-						UE_LOG(LogRivermaxMedia, Warning, TEXT("Rendering unexpected frame %u, when frame %u was expected."), Frame->FrameNumber, FrameExpectation.FrameNumber);
-					}
-					break;
-				}
+			// Our goal here is to wait until the expected frame is available to be used (received) unless there is a timeout
+			if (WaitConditionFunction(Sample))
+			{
+				return true;
+			}
 
+			{
+				FPlatformProcess::SleepNoStats(SleepTimeSeconds);
+				if (bCanTimeout && ((FPlatformTime::Seconds() - StartTimeSeconds) > TimeoutSeconds))
 				{
-					FPlatformProcess::SleepNoStats(SleepTimeSeconds);
-					if (bCanTimeout && ((FPlatformTime::Seconds() - StartTimeSeconds) > TimeoutSeconds))
-					{
-						UE_LOG(LogRivermaxMedia, Error, TEXT("Timed out waiting for frame %u."), FrameExpectation.FrameNumber);
-						break;
-					}
+					UE_LOG(LogRivermaxMedia, Error, TEXT("Timed out waiting for frame %u."), Sample->GetFrameNumber());
+					return false;
 				}
 			}
 		}
 	}
 
-	void FRivermaxMediaPlayer::SampleUploadSetupRenderThreadMode(const FFrameExpectation& NextFrameExpectations, FSampleConverterOperationSetup& OutConverterSetup)
+	void FRivermaxMediaPlayer::SampleUploadSetupTaskThreadMode(TSharedPtr<FRivermaxMediaTextureSample> Sample, FSampleConverterOperationSetup& OutConverterSetup)
 	{
-		TSharedPtr<FRivermaxSampleWrapper> SampleWrapper = SamplePool[NextFrameExpectations.FrameIndex];
-		if (bDoesStreamSupportsGPUDirect)
-		{
-			OutConverterSetup.GetGPUBufferFunc = [SampleWrapper]() { return SampleWrapper->Sample->GetGPUBuffer(); };
+		// We will always be providing a buffer already located on the GPU even when not using gpudirect
+		// Once a frame has arrived on system, we will upload it to the allocated gpu buffer.
+		OutConverterSetup.GetGPUBufferFunc = [SamplePtr = Sample]() { return SamplePtr->GetGPUBuffer(); };
 
-			// Setup requirements for sample to be ready to be rendered
-			FRHICommandListImmediate& RHICmdList = FRHICommandListImmediate::Get();
-			RHICmdList.EnqueueLambda([NextFrameExpectations, this](FRHICommandList&)
-			{
-				FWaitConditionFunc ReadyToRenderConditionFunc = [](const TSharedPtr<FRivermaxSampleWrapper>& Sample)
-				{
-					return Sample->bIsReadyToRender.load();
-				};
+		const uint64 NextFrameExpectations = (IS_FRAMELOCKED()) ? TO_EXPECTED_SAMPLE_FRAME_NUM(GFrameCounterRenderThread) : Sample->GetFrameNumber();
 
-				TRACE_CPUPROFILER_EVENT_SCOPE(RmaxWaitSampleReadyness);
-				constexpr bool bCanTimeout = true;
-				WaitForSample(NextFrameExpectations, MoveTemp(ReadyToRenderConditionFunc), bCanTimeout);
-			});
-		}
-		else
-		{
-			// For system memory path, we wait for the sample to be received before returning it.
-			OutConverterSetup.GetSystemBufferFunc = [NextFrameExpectations, SampleWrapper, this]()
+		OutConverterSetup.PreConvertFunc = [SamplePtr = Sample, NextFrameExpectations, this](const FRDGBuilder& GraphBuilder)
 			{
-				FWaitConditionFunc ReceivedConditionFunc = [](const TSharedPtr<FRivermaxSampleWrapper>& Sample)
+				TRACE_CPUPROFILER_EVENT_SCOPE_TEXT(*FString::Printf(TEXT("Rmax::01_PreConvertFunction %u"), (SamplePtr->GetFrameNumber())));
+				// When GPUDirect is not involved, we have an extra step to do. We need to wait for the sample to be received
+				// but also initiate the memcopy to gpu memory for it to be rendered
+				if (bStreamSupportsGPUDirect == false)
 				{
-					return Sample->ReceptionState == ESampleReceptionState::Received;
-				};
+					constexpr uint32 Offset = 0;
+					const uint32 Size = SamplePtr->GetGPUBuffer()->GetSize();
 
-				TRACE_CPUPROFILER_EVENT_SCOPE(RmaxWaitSampleReception);
-				constexpr bool bCanTimeout = true;
-				WaitForSample(NextFrameExpectations, MoveTemp(ReceivedConditionFunc), bCanTimeout);
+					SamplePtr->SetAwaitingForGPUTransfer(true);
 
-				// We can return the system buffer once we know it has arrived
-				return SampleWrapper->Sample->GetBuffer();
-			};
-		}
-	}
+					// Todo: Find a way to keep the buffer permanently locked.
+					if (SamplePtr->LockedMemory == nullptr)
+					{
+						SamplePtr->LockedMemory = GraphBuilder.RHICmdList.LockBuffer(SamplePtr->GetGPUBuffer()->GetRHI(), Offset, Size, EResourceLockMode::RLM_WriteOnly_NoOverwrite);
+					}
 
-	void FRivermaxMediaPlayer::SampleUploadSetupTaskThreadMode(const FFrameExpectation& NextFrameExpectations, FSampleConverterOperationSetup& OutConverterSetup)
-	{
-		TSharedPtr<FRivermaxSampleWrapper> SampleWrapper = SamplePool[NextFrameExpectations.FrameIndex];
+					++TasksInFlight;
+					UE::Tasks::Launch(UE_SOURCE_LOCATION,
+						[Size, SamplePtr, NextFrameExpectations, this]()
+						{
+							ON_SCOPE_EXIT
+							{
+								--TasksInFlight;
+							};
+							TRACE_CPUPROFILER_EVENT_SCOPE_TEXT(*FString::Printf(TEXT("Rmax::01_CopyPixels %u"), (NextFrameExpectations)));
 
-		// We will always be providing a buffer already located on the GPU even when not using gpudirect
-		// Once a frame has arrived on system, we will upload it to the allocated gpu buffer.
-		OutConverterSetup.GetGPUBufferFunc = [SampleWrapper]() { return SampleWrapper->Sample->GetGPUBuffer(); };
+							// Indicates if we sample reception started on the input thread.
+							bool bStreamReceptionStarted = true;
 
-		OutConverterSetup.PreConvertFunc = [NextFrameExpectations, SampleWrapper, this](const FRDGBuilder& GraphBuilder)
-		{
-			// When GPUDirect is not involved, we have an extra step to do. We need to wait for the sample to be received
-			// but also initiate the memcopy to gpu memory for it to be rendered
-			if (bDoesStreamSupportsGPUDirect == false)
-			{
-				constexpr uint32 Offset = 0;
-				const uint32 Size = SampleWrapper->Sample->GetGPUBuffer()->GetSize();
-				SampleWrapper->LockedMemory = GraphBuilder.RHICmdList.LockBuffer(SampleWrapper->Sample->GetGPUBuffer()->GetRHI(), Offset, Size, EResourceLockMode::RLM_WriteOnly);
+							// In frame locked mode the packets potentially haven't started being received yet. 
+							// This waits for frame locked samples to start writing packets and then for the completion of the reception.
+							if (IS_FRAMELOCKED())
+							{
+								TRACE_CPUPROFILER_EVENT_SCOPE_TEXT(*FString::Printf(TEXT("Rmax::01_WaitForReception %u"), (NextFrameExpectations)));
+
+								// Waiting to start receiving.
+								const double TimeoutSeconds = CVarRivermaxDefaultThreadBlockTimeoutSec.GetValueOnAnyThread();
+								const double StartTimeSeconds = FPlatformTime::Seconds();
+								uint64 CurrentSampleFrameNumber = SamplePtr->GetFrameNumber();
+
+								while (CurrentSampleFrameNumber < NextFrameExpectations)
+								{
+									FPlatformProcess::SleepNoStats(SleepTimeSeconds);
+									if ((FPlatformTime::Seconds() - StartTimeSeconds) > TimeoutSeconds)
+									{
+										UE_LOG(LogRivermaxMedia, Warning, TEXT("Timed out waiting for frame #%u to start being received."), NextFrameExpectations);
+										bStreamReceptionStarted = false;
+										break;
+									}
+									CurrentSampleFrameNumber = SamplePtr->GetFrameNumber();
+								}
+
+								// At this point it is known that the sample reception has started or not.
+								if (bStreamReceptionStarted)
+								{
+									TRACE_CPUPROFILER_EVENT_SCOPE(RmaxWaitSampleReception);
+									TRACE_CPUPROFILER_EVENT_SCOPE_TEXT(*FString::Printf(TEXT("Rmax::02_WaitSampleReceived %u"), (NextFrameExpectations)));
+									// We wait for 2 Media Player frames to timeout. Reason: Due to 2110 standard anything more than 1 frame is going to be discarded.
+									// Since we are guaranteed to have started receiving the buffer we shouldn't wait longer than that.
+									const double TwoFrameTime = 2. / GetFrameRate().AsDecimal();
+									SamplePtr->GetSampleReceivedEvent()->Wait(FTimespan::FromSeconds(TwoFrameTime));
+								}
+							}
 
-				++TasksInFlight;
-				UE::Tasks::Launch(UE_SOURCE_LOCATION,
-					[NextFrameExpectations, Size, SampleWrapper, this]()
-					{
-						ON_SCOPE_EXIT
-						{
-							--TasksInFlight;
-						};
+							if (SamplePtr->GetFrameNumber() != NextFrameExpectations)
+							{
+								UE_LOG(LogRivermaxMedia, Warning, TEXT("Rendering unexpected frame %u, when frame %u was expected."), SamplePtr->GetFrameNumber(), NextFrameExpectations);
+							}
 
-						FWaitConditionFunc ReceivedConditionFunc = [](const TSharedPtr<FRivermaxSampleWrapper>& Sample)
-						{
-							return Sample->ReceptionState == ESampleReceptionState::Received;
-						};
+							// In case reception failed mid way.
+							if (bStreamReceptionStarted && SamplePtr->GetReceptionState() != IRivermaxSample::ESampleState::ReceptionError)
+							{
+								TRACE_CPUPROFILER_EVENT_SCOPE_TEXT(*FString::Printf(TEXT("Rmax::03_CopyBufferToGPU %u"), (SamplePtr->GetFrameNumber())));
+								FMemory::Memcpy(SamplePtr->LockedMemory, SamplePtr->GetOrRequestBuffer(Size), Size);
+							}
+							else
+							{
+								UE_LOG(LogRivermaxMedia, Warning, TEXT("Incomplete or failed pixels will be rendered for frame %u"), NextFrameExpectations);
+							}
 
-						{
-							TRACE_CPUPROFILER_EVENT_SCOPE(RmaxWaitSampleReception);
-							constexpr bool bCanTimeout = true;
-							WaitForSample(NextFrameExpectations, MoveTemp(ReceivedConditionFunc), bCanTimeout);
+							// Signals the RHI thread that the GPU transfer has completed.
+							SamplePtr->SetAwaitingForGPUTransfer(false);
 						}
+					);
+				}
 
-						// In case the wait failed, make sure it's received before copying
-						if (SamplePool[NextFrameExpectations.FrameIndex]->ReceptionState == ESampleReceptionState::Received)
-						{
-							TRACE_CPUPROFILER_EVENT_SCOPE(RmaxSampleUpload);
-							FMemory::Memcpy(SampleWrapper->LockedMemory, SampleWrapper->Sample->GetBuffer(), Size);
-						}
-						
-						// We always consider the sample ready to be rendered even if we haven't copied something over (timeout)
-						// RHI commands have already been submitted to look for that frame so we can't back out at this point
-						SampleWrapper->bIsReadyToRender = true;
-					}
-				);
-			}
-			
-			// Setup requirements for sample to be ready to be rendered
-			RHI_BREADCRUMB_EVENT_STAT(GraphBuilder.RHICmdList, Rmax_WaitForPixels, "Rmax_WaitForPixels");
-			SCOPED_GPU_STAT(GraphBuilder.RHICmdList, Rmax_WaitForPixels);
 
-			// Since we are going to enqueue a lambda that can potentially sleep in the RHI thread if the pixels haven't arrived,
-			// we dispatch the existing commands (including the draw event start timing in the SCOPED_DRAW_EVENT above) before any potential sleep.
-			GraphBuilder.RHICmdList.ImmediateFlush(EImmediateFlushType::DispatchToRHIThread);
+				RHI_BREADCRUMB_EVENT_STAT(GraphBuilder.RHICmdList, Rmax_WaitForPixels, "Rmax::WaitForPixels");
+				SCOPED_GPU_STAT(GraphBuilder.RHICmdList, Rmax_WaitForPixels);
+				// Since we are going to enqueue a lambda that can potentially sleep in the RHI thread if the pixels haven't arrived,
+				// we dispatch the existing commands (including the draw event start timing in the SCOPED_DRAW_EVENT above) before any potential sleep.
+				GraphBuilder.RHICmdList.ImmediateFlush(EImmediateFlushType::DispatchToRHIThread);
 
-			GraphBuilder.RHICmdList.EnqueueLambda(
-				[NextFrameExpectations, this](FRHICommandListImmediate&)
-				{
-					FWaitConditionFunc ReadyToRenderConditionFunc = [](const TSharedPtr<FRivermaxSampleWrapper>& Sample)
+				// Setup requirements for sample to be ready to be rendered
+				GraphBuilder.RHICmdList.EnqueueLambda(
+					[SamplePtr, NextFrameExpectations, this](FRHICommandListImmediate&)
 					{
-						return Sample->bIsReadyToRender.load();
-					};
+						FWaitConditionFunc ReadyToRenderConditionFunc = [](const TSharedPtr<FRivermaxMediaTextureSample>& Sample)
+							{
+								return Sample->IsAwaitingForGPUTransfer() == false;
+							};
 
-					TRACE_CPUPROFILER_EVENT_SCOPE(RmaxWaitSampleReadyness);
+						TRACE_CPUPROFILER_EVENT_SCOPE_TEXT(*FString::Printf(TEXT("Rmax::04_WaitForUploadToFinish %u"), (NextFrameExpectations)));
+
+						// Only accept timeouts here if we're going through the system memory in order to avoid timing out while doing
+						// the memcopy (buffer upload)
+						bool bCanTimeout = bStreamSupportsGPUDirect;
+						WaitForSample(SamplePtr, MoveTemp(ReadyToRenderConditionFunc), bCanTimeout);
+					}
+				);
 
-					// Only accept timeouts here if we're going through the system memory in order to avoid timing out while doing
-					// the memcopy (buffer upload)
-					bool bCanTimeout = bDoesStreamSupportsGPUDirect;
-					WaitForSample(NextFrameExpectations, MoveTemp(ReadyToRenderConditionFunc), bCanTimeout);
+				// Final step, if the memory was locked (non gpu direct), enqueue unlock after the wait for sample in order to render it
+				if (SamplePtr->LockedMemory && ensure(!bStreamSupportsGPUDirect))
+				{
+					GraphBuilder.RHICmdList.UnlockBuffer(SamplePtr->GetGPUBuffer()->GetRHI());
 				}
-			);
+			};
 
-			// Final step, if the memory was locked (non gpu direct), enqueue unlock after the wait for sample in order to render it
-			if (SampleWrapper->LockedMemory && ensure(!bDoesStreamSupportsGPUDirect))
+		// Setup post sample usage pass 
+		OutConverterSetup.PostConvertFunc = [SamplePtr = Sample, this](FRDGBuilder& GraphBuilder)
 			{
-				GraphBuilder.RHICmdList.UnlockBuffer(SampleWrapper->Sample->GetGPUBuffer()->GetRHI());
-			}
-		};
+				PostSampleUsage(GraphBuilder, SamplePtr);
+			};
 	}
 
 }
diff --git a/Engine/Plugins/VirtualProduction/Rivermax/RivermaxMedia/Source/RivermaxMedia/Private/RivermaxMediaPlayer.h b/Engine/Plugins/VirtualProduction/Rivermax/RivermaxMedia/Source/RivermaxMedia/Private/RivermaxMediaPlayer.h
--- a/Engine/Plugins/VirtualProduction/Rivermax/RivermaxMedia/Source/RivermaxMedia/Private/RivermaxMediaPlayer.h
+++ b/Engine/Plugins/VirtualProduction/Rivermax/RivermaxMedia/Source/RivermaxMedia/Private/RivermaxMediaPlayer.h
@@ -22,7 +22,6 @@ namespace  UE::RivermaxMedia
 {
 	using namespace UE::RivermaxCore;
 	
-	class FRivermaxMediaTextureSamples;
 	class FRivermaxMediaTextureSample;
 	class FRivermaxMediaTextureSamplePool;
 
@@ -45,55 +44,52 @@ namespace  UE::RivermaxMedia
 		virtual ~FRivermaxMediaPlayer();
 
 	public:
-		
-		/** Sample update called by converter to setup sample to render */
-		bool LateUpdateSetupSample(FSampleConverterOperationSetup& OutConverterSetup);
+		//~ Begin FMediaIOCorePlayerBase interface
+
+		/** Called by the sample converter to setup rendering commands to convert this sample into texture. */
+		virtual bool JustInTimeSampleRender_RenderThread(FRHICommandListImmediate& RHICmdList, FTextureRHIRef& InDestinationTexture, TSharedPtr<FMediaIOCoreTextureSampleBase>& JITRProxySample) override;
+	protected:
+		virtual TSharedPtr<FMediaIOCoreTextureSampleBase> PickSampleToRenderFramelocked_RenderThread(const FFrameInfo& InFrameInformation) override;
 
+		//~ End FMediaIOCorePlayerBase interface
+
+	public: 
 		//~ Begin IMediaPlayer interface
 		virtual void Close() override;
 		virtual FGuid GetPlayerPluginGUID() const override;
 		virtual bool Open(const FString& Url, const IMediaOptions* Options) override;
-		virtual void TickTimeManagement() override;
-		virtual void TickFetch(FTimespan DeltaTime, FTimespan Timecode) override;
 		virtual void TickInput(FTimespan DeltaTime, FTimespan Timecode) override;
-		virtual IMediaSamples& GetSamples() override;
-		virtual FString GetStats() const override;
-		virtual bool GetPlayerFeatureFlag(EFeatureFlag flag) const override;
-		virtual bool SetRate(float Rate) override;
 		//~ End IMediaPlayer interface
 
 		//~ Begin ITimedDataInput interface
 #if WITH_EDITOR
 		virtual const FSlateBrush* GetDisplayIcon() const override;
 #endif
+		//~ End ITimedDataInput interface
+
+		ERivermaxMediaSourcePixelFormat GetDesiredPixelFormat()
+		{
+			return DesiredPixelFormat;
+		};
 
 		//~ Begin IRivermaxInputStreamListener interface
 		virtual void OnInitializationCompleted(const FRivermaxInputInitializationResult& Result) override;
-		virtual bool OnVideoFrameRequested(const FRivermaxInputVideoFrameDescriptor& FrameInfo, FRivermaxInputVideoFrameRequest& OutVideoFrameRequest) override;
-		virtual void OnVideoFrameReceived(const FRivermaxInputVideoFrameDescriptor& FrameInfo, const FRivermaxInputVideoFrameReception& ReceivedVideoFrame) override;
-		virtual void OnVideoFrameReceptionError(const FRivermaxInputVideoFrameDescriptor& FrameInfo) override;
+		virtual TSharedPtr<IRivermaxVideoSample> OnVideoFrameRequested(const FRivermaxInputVideoFrameDescriptor& FrameInfo) override;
+		virtual void OnVideoFrameReceived(TSharedPtr<IRivermaxVideoSample> InReceivedVideoFrame) override;
+		virtual void OnVideoFrameReceptionError(TSharedPtr<IRivermaxVideoSample> InReceivedVideoFrame) override;
 		virtual void OnStreamError() override;
 		virtual void OnVideoFormatChanged(const FRivermaxInputVideoFormatChangedInfo& NewFormatInfo) override;
 		//~ End IRivermaxInputStreamListener interface
 
-	protected:
-
-		/**
-		 * Process pending audio and video frames, and forward them to the sinks.
-		 */
-		void ProcessFrame();
-
 	protected:
 
 		//~ Begin FMediaIOCorePlayerBase interface
 		virtual bool IsHardwareReady() const override;
 		virtual void SetupSampleChannels() override;
-		virtual TSharedPtr<FMediaIOCoreTextureSampleBase> AcquireTextureSample_AnyThread() const override
-		{
-			// This needs to be fixed once FRivermaxMediaTextureSample is inherited from FMediaIOCoreTextureSampleBase
-			// return MakeShared<FRivermaxMediaTextureSample>();
-			return nullptr;
-		}
+		virtual TSharedPtr<FMediaIOCoreTextureSampleConverter> CreateTextureSampleConverter() const override;
+
+		virtual TSharedPtr<FMediaIOCoreTextureSampleBase> AcquireTextureSample_AnyThread() const override;
+		
 		//~ End FMediaIOCorePlayerBase interface
 
 	private:
@@ -104,101 +100,28 @@ namespace  UE::RivermaxMedia
 		/** Allocates the sample pool used to receive incoming data */
 		void AllocateBuffers(const FIntPoint& InResolution);
 
-		/** Wrapper struct holding information about frame expected to be rendered */
-		struct FFrameExpectation
-		{
-			/** Frame number we are expecting to find at ExpectedIndex */
-			uint32 FrameNumber = 0;
-
-			/** Index where rendering should pickup its sample from */
-			uint32 FrameIndex = 0;
-		};
-
-		enum class ESampleReceptionState : uint8
-		{
-			// Sample is ready to be requested by rivermax stream
-			Available,
-
-			// Sample has been requested and is being received
-			Receiving,
-
-			// Sample has been received and is ready to be rendered
-			Received,
-		};
-
-		/** Wrapper around a sample holding more information about its state */
-		struct FRivermaxSampleWrapper
-		{
-			/** State of this sample */
-			std::atomic<ESampleReceptionState> ReceptionState = ESampleReceptionState::Available;
-
-			/** True when queued for rendering. Will be false once fence has been written, after shader usage. */
-			std::atomic<bool> bIsPendingRendering = false;
-
-			/** True if sample can be rendered. For non-gpudirect, we need to copy sample from system memory to gpu to be ready */
-			std::atomic<bool> bIsReadyToRender = false;
-
-			/** Locked memory of gpu buffer when uploading */
-			void* LockedMemory = nullptr;
-
-			/** Actual sample buffer container used by media framework */
-			TSharedPtr<FRivermaxMediaTextureSample> Sample;
-
-			/** Write fence enqueued after sample conversion to know when it's ready to be reused */
-			FGPUFenceRHIRef SampleConversionFence;
-
-			/** Frame number of this sample */
-			uint32 FrameNumber = 0;
-
-			/** Timestamp of this sample */
-			uint32 Timestamp = 0;
-		};
-
-		/** Buffer upload setup that will block render thread while waiting for sample and uploading it */
-		void SampleUploadSetupRenderThreadMode(const FFrameExpectation& FrameExpectation, FSampleConverterOperationSetup& OutConverterSetup);
-		
 		/** Buffer upload setup that will wait on its own task to wait for sample and do the upload */
-		void SampleUploadSetupTaskThreadMode(const FFrameExpectation& FrameExpectation, FSampleConverterOperationSetup& OutConverterSetup);
+		void SampleUploadSetupTaskThreadMode(TSharedPtr<FRivermaxMediaTextureSample> Sample, FSampleConverterOperationSetup& OutConverterSetup);
 
 		/** Function waiting for the expected frame to be received. */
-		using FWaitConditionFunc = TUniqueFunction<bool(const TSharedPtr<FRivermaxSampleWrapper>&)>;
-		void WaitForSample(const FFrameExpectation& FrameExpectation, FWaitConditionFunc WaitConditionFunction, bool bCanTimeout);
+		using FWaitConditionFunc = TUniqueFunction<bool(const TSharedPtr<FRivermaxMediaTextureSample>&)>;
+		bool WaitForSample(TSharedPtr<FRivermaxMediaTextureSample> Sample, FWaitConditionFunc WaitConditionFunction, bool bCanTimeout);
 		
 		/** Called after sample was converted / rendered. Used write a fence to detect when sample is reusable */
-		void PostSampleUsage(FRDGBuilder& GraphBuilder, const FFrameExpectation& FrameExpectation);
-
-		/** Provides next frame to render expectations in terms of number and location in the pool.  */
-		bool GetNextExpectedFrameInfo(FFrameExpectation& OutExpectation);
-		bool GetNextExpectedFrameInfoForFramelock(FFrameExpectation& OutExpectation);
-		bool GetNextExpectedFrameInfoForLatest(FFrameExpectation& OutExpectation);
-
-		/** Provides next requested frame (by receiver stream) location in the pool */
-		bool GetFrameRequestedIndex(const FRivermaxInputVideoFrameDescriptor& FrameInfo, uint32& OutExpectedIndex);
-		bool GetFrameRequestedIndexForFramelock(const FRivermaxInputVideoFrameDescriptor& FrameInfo, uint32& OutExpectedIndex);
-		bool GetFrameRequestedIndexForLatest(const FRivermaxInputVideoFrameDescriptor& FrameInfo, uint32& OutExpectedIndex);
-
-		/** Verifies if a frame should be skipped, looking at gaps we might have had in the reception of frames */
-		bool IsFrameSkipped(uint32 FrameNumber) const;
-
-		/** Looks at last frame rendered and tries to clean skipped frames container */
-		void TryClearSkippedInterval(uint32 LastFrameRendered);
-
+		void PostSampleUsage(FRDGBuilder& GraphBuilder, TSharedPtr<FRivermaxMediaTextureSample> Sample);
+		
 		/** Whether player is ready to play */
 		bool IsReadyToPlay() const;
 
 		/** Waits for tasks in flight and flushes render commands before cleaning our ressources */
 		void WaitForPendingTasks();
 
+		/** Used to create texture for color encoding conversion. */
+		TRefCountPtr<FRHITexture> CreateIntermediateRenderTarget(FRHICommandListImmediate& RHICmdList, const FIntPoint& InDim, EPixelFormat InPixelFormat, bool bInSRGB);
 
 	private:
-
-		TSharedPtr<FRivermaxSampleWrapper> RivermaxThreadCurrentTextureSample;
-
-		/** Common gpu buffer to use when we render a sample */
-		TRefCountPtr<FRDGPooledBuffer> CommonGPUBuffer;
-
 		/** Size of the sample pool. */
-		int32 MaxNumVideoFrameBuffer;
+		static constexpr uint32 kMaxNumVideoFrameBuffer = 5;
 
 		/** Current state of the media player. */
 		EMediaState RivermaxThreadNewState;
@@ -206,73 +129,26 @@ namespace  UE::RivermaxMedia
 		/** Options used to configure the stream. i.e.  */
 		FRivermaxInputStreamOptions StreamOptions;
 
-		/** Whether the input is in sRGB and can have a ToLinear conversion. */
-		bool bIsSRGBInput;
-
-		/** Which field need to be capture. */
-		bool bUseVideo;
-		bool bVerifyFrameDropCount;
-
 		/** Maps to the current input Device */
 		TUniquePtr<IRivermaxInputStream> InputStream;
 
-		/** Used to flag which sample types we advertise as supported for timed data monitoring */
-		EMediaIOSampleType SupportedSampleTypes;
-
-		/** Flag to indicate that pause is being requested */
-		std::atomic<bool> bPauseRequested;
-
 		/** Pixel format provided by media source */
 		ERivermaxMediaSourcePixelFormat DesiredPixelFormat = ERivermaxMediaSourcePixelFormat::RGB_10bit;
 
-		/** Pool of samples where incoming ones are written to and at render time, we select our candidate */
-		TArray<TSharedPtr<FRivermaxSampleWrapper>> SamplePool;
+		/** 
+		* Pool of samples. The pool has full management of shared pointers.
+		* Unreferenced shared pointers automatically return back to the pool and released only when pool is destroyed. 
+		*/
+		TUniquePtr<FRivermaxMediaTextureSamplePool> VideoTextureSamplePool;
 
-		/** Used to ensure that JustInTimeSampleRender is only run once per frame */
-		uint32 LastFrameNumberThatUpdatedJustInTime = 0;
-
-		/** Critical section used when manipulating the skipped frame container */
-		mutable FCriticalSection SkippedFrameCriticalSection;
-
-		struct FFrameTracking
-		{
-			/** Whether a frame has been requested yet. Used to detect start of the stream and valid frame number expectations */
-			bool bWasFrameRequested = false;
-			
-			/** Last frame number sent to render. Used to release received frames that will never be picked up early on */
-			TOptional<uint32> LastFrameRendered;
-
-			/** Last frame expectations used */
-			FFrameExpectation LastFrameExpectation;
-			
-			/** Last index requested at the start of a frame reception. In latest mode, this just goes up incrementally */
-			uint32 LastFrameRequestedIndex = 0;
-
-			/** Used to track frame reception and detect gaps in received frames */
-			uint32 LastFrameNumberRequested = 0;
-
-			/** Used to detect frames to be rendered with invalid expectations. i.e Looking for a frame number that we will never get */
-			uint32 FirstFrameRequested = 0;
-		};
-		FFrameTracking FrameTracking;
-
-		/** Mode the player is in. Latest, framelocked, etc... */
-		ERivermaxPlayerMode PlayerMode = ERivermaxPlayerMode::Latest;
+		/** Sample that input stream should write to in framelocking mode. */
+		TStaticArray<TSharedPtr<FRivermaxMediaTextureSample>, kMaxNumVideoFrameBuffer> FrameLockedSamples;
 
 		/** Number of tasks currently in progress. Used during shutdown to know when we are good to continue */
 		std::atomic<uint32> TasksInFlight = 0;
-	
-		/** Used to detect skipped / missing frames in our reception and early exit the waiting in case we are waiting for a frame that will never come */
-		TArray<TInterval<uint32>> SkippedFrames;
-
-		/** Used in framelock mode. How far behind frame counter we will try to find a matching frame */
-		uint32 FrameLatency = 0;
-
-		/** Special sample container, holding a single one used to trigger media framework rendering */
-		TUniquePtr<FRivermaxMediaTextureSamples> MediaSamples;
 
 		/** Whether the created stream supports GPUDirect. Will be confirmed after initialization. */
-		bool bDoesStreamSupportsGPUDirect = false;
+		bool bStreamSupportsGPUDirect = false;
 
 		/** Time to sleep when waiting for an operation to complete */
 		static constexpr double SleepTimeSeconds = 50.0 * 1E-6;
@@ -285,6 +161,15 @@ namespace  UE::RivermaxMedia
 		
 		/** Whether the player follows resolution detected by our stream, adjusting texture size as required */
 		bool bFollowsStreamResolution = true;
+
+		/** Used to make sure that the player didn't accidentally skipped the reception of any frames. */
+		uint32 LastFrameToAttemptReception = 0;
+
+		/** Critical section used when accessing ProxySampleDummy. */
+		mutable FCriticalSection ProxySampleAccessCriticalSection;
+
+		/** This is the proxy sample contains all common settings for texture samples for this player. */
+		TSharedPtr<FRivermaxMediaTextureSample> ProxySampleDummy;
 	};
 }
 
diff --git a/Engine/Plugins/VirtualProduction/Rivermax/RivermaxMedia/Source/RivermaxMedia/Private/RivermaxMediaSource.cpp b/Engine/Plugins/VirtualProduction/Rivermax/RivermaxMedia/Source/RivermaxMedia/Private/RivermaxMediaSource.cpp
--- a/Engine/Plugins/VirtualProduction/Rivermax/RivermaxMedia/Source/RivermaxMedia/Private/RivermaxMediaSource.cpp
+++ b/Engine/Plugins/VirtualProduction/Rivermax/RivermaxMedia/Source/RivermaxMedia/Private/RivermaxMediaSource.cpp
@@ -4,26 +4,28 @@
 
 #include "MediaIOCorePlayerBase.h"
 #include "RivermaxMediaSourceOptions.h"
+#include "RivermaxMediaUtils.h"
 
 
 /*
  * IMediaOptions interface
  */
 
+URivermaxMediaSource::URivermaxMediaSource() : UCaptureCardMediaSource()
+{
+	Deinterlacer = nullptr;
+	bOverrideSourceEncoding = false;
+	bOverrideSourceColorSpace = false;
+	bRenderJIT = true;
+	EvaluationType = EMediaIOSampleEvaluationType::Latest;
+}
+
 bool URivermaxMediaSource::GetMediaOption(const FName& Key, bool DefaultValue) const
 {
-	if (Key == RivermaxMediaOption::SRGBInput)
-	{
-		return bIsSRGBInput;
-	}
 	if (Key == RivermaxMediaOption::UseGPUDirect)
 	{
 		return bUseGPUDirect;
 	}
-	if (Key == RivermaxMediaOption::ZeroLatency)
-	{
-		return bUseZeroLatency;
-	}
 	if (Key == RivermaxMediaOption::OverrideResolution)
 	{
 		return bOverrideResolution;
@@ -58,10 +60,6 @@ int64 URivermaxMediaSource::GetMediaOption(const FName& Key, int64 DefaultValue)
 	{
 		return Resolution.Y;
 	}
-	else if (Key == RivermaxMediaOption::PlayerMode)
-	{
-		return (int64)PlayerMode;
-	}
 	return Super::GetMediaOption(Key, DefaultValue);
 }
 
@@ -88,7 +86,6 @@ bool URivermaxMediaSource::HasMediaOption(const FName& Key) const
 		(Key == RivermaxMediaOption::StreamAddress) ||
 		(Key == RivermaxMediaOption::Port) ||
 		(Key == RivermaxMediaOption::PixelFormat) ||
-		(Key == RivermaxMediaOption::SRGBInput) ||
 		(Key == RivermaxMediaOption::UseGPUDirect) ||
 		(Key == FMediaIOCoreMediaOption::FrameRateNumerator) ||
 		(Key == FMediaIOCoreMediaOption::FrameRateDenominator) ||
@@ -102,6 +99,57 @@ bool URivermaxMediaSource::HasMediaOption(const FName& Key) const
 	return Super::HasMediaOption(Key);
 }
 
+void URivermaxMediaSource::Serialize(FArchive& Ar_Asset)
+{
+	Super::Serialize(Ar_Asset);
+
+	Ar_Asset.UsingCustomVersion(FRivermaxMediaVersion::GUID);
+}
+
+void URivermaxMediaSource::PostLoad()
+{
+	Super::PostLoad();
+	// We can only recover data during editor. Proprties will be fixed during cook.
+#if WITH_EDITORONLY_DATA
+	const int32 RivermaxMediaVersion = GetLinkerCustomVersion(FRivermaxMediaVersion::GUID);
+	PRAGMA_DISABLE_DEPRECATION_WARNINGS;
+
+	// Rivermax was merged with other capture card media sources and these properties 
+	// are either duplicates or provide limited functionality compared to the inherited UI.
+	if (RivermaxMediaVersion < FRivermaxMediaVersion::BeforeCustomVersionAdded)
+	{
+		if (bUseZeroLatency_DEPRECATED == false)
+		{
+			bUseTimeSynchronization = true;
+			FrameDelay = 1;
+		}
+
+		if (bIsSRGBInput_DEPRECATED == true)
+		{
+			bOverrideSourceEncoding = true;
+			OverrideSourceEncoding = EMediaIOCoreSourceEncoding::sRGB;
+		}
+
+		if (PlayerMode_DEPRECATED == ERivermaxPlayerMode_DEPRECATED::Framelock)
+		{
+			EvaluationType = EMediaIOSampleEvaluationType::Timecode;
+			bUseTimeSynchronization = true;
+			bFramelock = true;
+		}
+		else
+		{
+			EvaluationType = EMediaIOSampleEvaluationType::Latest;
+			bFramelock = false;
+		}
+
+		Modify();
+	}
+
+	PRAGMA_ENABLE_DEPRECATION_WARNINGS;
+#endif
+
+}
+
 /*
  * UMediaSource interface
  */
diff --git a/Engine/Plugins/VirtualProduction/Rivermax/RivermaxMedia/Source/RivermaxMedia/Private/RivermaxMediaSourceOptions.h b/Engine/Plugins/VirtualProduction/Rivermax/RivermaxMedia/Source/RivermaxMedia/Private/RivermaxMediaSourceOptions.h
--- a/Engine/Plugins/VirtualProduction/Rivermax/RivermaxMedia/Source/RivermaxMedia/Private/RivermaxMediaSourceOptions.h
+++ b/Engine/Plugins/VirtualProduction/Rivermax/RivermaxMedia/Source/RivermaxMedia/Private/RivermaxMediaSourceOptions.h
@@ -11,9 +11,6 @@ namespace RivermaxMediaOption
 	static const FName StreamAddress("StreamAddress");
 	static const FName Port("Port");
 	static const FName PixelFormat("PixelFormat");
-	static const FName SRGBInput("SRGBInput");
 	static const FName UseGPUDirect("UseGPUDirect");
-	static const FName PlayerMode("PlayerMode");
-	static const FName ZeroLatency("ZeroLatency");
 	static const FName OverrideResolution("OverrideResolution");
 }
diff --git a/Engine/Plugins/VirtualProduction/Rivermax/RivermaxMedia/Source/RivermaxMedia/Private/RivermaxMediaTextureSample.cpp b/Engine/Plugins/VirtualProduction/Rivermax/RivermaxMedia/Source/RivermaxMedia/Private/RivermaxMediaTextureSample.cpp
--- a/Engine/Plugins/VirtualProduction/Rivermax/RivermaxMedia/Source/RivermaxMedia/Private/RivermaxMediaTextureSample.cpp
+++ b/Engine/Plugins/VirtualProduction/Rivermax/RivermaxMedia/Source/RivermaxMedia/Private/RivermaxMediaTextureSample.cpp
@@ -6,14 +6,20 @@
 #include "RivermaxMediaPlayer.h"
 #include "RivermaxMediaTextureSampleConverter.h"
 #include "RivermaxMediaUtils.h"
+#include "Templates/RefCounting.h"
 
 namespace UE::RivermaxMedia
 {
 
 FRivermaxMediaTextureSample::FRivermaxMediaTextureSample()
-	: Converter(MakePimpl<FRivermaxMediaTextureSampleConverter>())
+	: FMediaIOCoreTextureSampleBase()
 {
-	
+	SampleReceivedEvent = FGenericPlatformProcess::GetSynchEventFromPool(false);
+}
+
+FRivermaxMediaTextureSample::~FRivermaxMediaTextureSample()
+{
+	FGenericPlatformProcess::ReturnSynchEventToPool(SampleReceivedEvent);
 }
 
 const FMatrix& FRivermaxMediaTextureSample::GetYUVToRGBMatrix() const
@@ -26,114 +32,45 @@ IMediaTextureSampleConverter* FRivermaxMediaTextureSample::GetMediaTextureSample
 	return Converter.Get();
 }
 
-bool FRivermaxMediaTextureSample::IsOutputSrgb() const
-{
-	// We always do the sRGB to Linear conversion in shader if specified by the source
-	// If true is returned here, MediaTextureResource will create (try) a sRGB texture which only works for 8 bits
-	return false;
-}
-
-bool FRivermaxMediaTextureSample::ConfigureSample(const FSampleConfigurationArgs& Args)
+bool FRivermaxMediaTextureSample::InitializeJITR(const FMediaIOCoreSampleJITRConfigurationArgs& Args)
 {
-	switch (Args.SampleFormat)
+	if (!FMediaIOCoreTextureSampleBase::InitializeJITR(Args))
 	{
-		case ERivermaxMediaSourcePixelFormat::RGB_12bit:
-			// Falls through
-		case ERivermaxMediaSourcePixelFormat::RGB_16bit_Float:
-		{
-			SampleFormat = EMediaTextureSampleFormat::FloatRGBA;
-			break;
-		}
-		case ERivermaxMediaSourcePixelFormat::RGB_10bit:
-			// Falls through
-		case ERivermaxMediaSourcePixelFormat::YUV422_10bit:
-		{
-			SampleFormat = EMediaTextureSampleFormat::CharBGR10A2;
-			break;
-		}
-		case ERivermaxMediaSourcePixelFormat::YUV422_8bit:
-			// Falls through
-		case ERivermaxMediaSourcePixelFormat::RGB_8bit:
-			// Falls through
-		default:
-		{
-			SampleFormat = EMediaTextureSampleFormat::CharBGRA;
-			break;
-		}
+		return false;
 	}
 
-	Converter->Setup(AsShared());
-
-	WeakPlayer = Args.Player;
-	InputFormat = Args.SampleFormat;
-	Dimension = FIntPoint(Args.Width, Args.Height);
-	Time = Args.Time;
-	Duration = FTimespan(1); // We keep our frame to have a duration of 1 frame
-	bIsInputSRGB = Args.bInIsSRGBInput;
+	ERivermaxMediaSourcePixelFormat DesiredPixelFormat = StaticCastSharedPtr<FRivermaxMediaPlayer>(Args.Player)->GetDesiredPixelFormat();
+	SetInputFormat(DesiredPixelFormat);
 
 	return true;
 }
 
-TRefCountPtr<FRDGPooledBuffer> FRivermaxMediaTextureSample::GetGPUBuffer() const
+void FRivermaxMediaTextureSample::CopyConfiguration(const TSharedPtr<FMediaIOCoreTextureSampleBase>& SourceSample)
 {
-	return GPUBuffer;
+	FMediaIOCoreTextureSampleBase::CopyConfiguration(SourceSample);
+	TSharedPtr<FRivermaxMediaTextureSample> RivermaxSample = StaticCastSharedPtr<FRivermaxMediaTextureSample>(SourceSample);
+	SetInputFormat(RivermaxSample->InputFormat);
 }
 
-void FRivermaxMediaTextureSample::InitializeGPUBuffer(const FIntPoint& InResolution, ERivermaxMediaSourcePixelFormat InSampleFormat)
+void FRivermaxMediaTextureSample::InitializeGPUBuffer(const FIntPoint& InResolution, ERivermaxMediaSourcePixelFormat InSampleFormat, bool bSupportsGPUDirect)
 {
 	using namespace UE::RivermaxMediaUtils::Private;
-
 	const FSourceBufferDesc BufferDescription = GetBufferDescription(InResolution, InSampleFormat);
 
 	FRDGBufferDesc RDGBufferDesc = FRDGBufferDesc::CreateStructuredDesc(BufferDescription.BytesPerElement, BufferDescription.NumberOfElements);
 	
-	// Required to share resource across different graphics API (DX, Cuda)
-	RDGBufferDesc.Usage |= EBufferUsageFlags::Shared;
-
-	TWeakPtr<FRivermaxMediaTextureSample> WeakSample = AsShared();
-	ENQUEUE_RENDER_COMMAND(RivermaxPlayerBufferCreation)(
-	[WeakSample, RDGBufferDesc](FRHICommandListImmediate& CommandList)
 	{
-		if (TSharedPtr<FRivermaxMediaTextureSample> Sample = WeakSample.Pin())
-		{
-			Sample->GPUBuffer = AllocatePooledBuffer(RDGBufferDesc, TEXT("RmaxInput Buffer"));
-		}
-	});
-}
-
-const void* FRivermaxMediaTextureSample::GetBuffer()
-{
-	return Buffer.GetData();
-}
-
-FIntPoint FRivermaxMediaTextureSample::GetDim() const
-{
-	return Dimension;
-}
-
-FTimespan FRivermaxMediaTextureSample::GetDuration() const
-{
-	return Duration;
-}
-
-EMediaTextureSampleFormat FRivermaxMediaTextureSample::GetFormat() const
-{
-	return SampleFormat;
-}
-
-FIntPoint FRivermaxMediaTextureSample::GetOutputDim() const
-{
-	return GetDim();
-}
+		// Required to share resource across different graphics API (DX, Cuda)
+		RDGBufferDesc.Usage |= EBufferUsageFlags::Shared;
+	}
 
-uint32 FRivermaxMediaTextureSample::GetStride() const
-{
-	return Stride;
-}
+	InputFormat = InSampleFormat;
+	ENQUEUE_RENDER_COMMAND(FRivermaxMediaTextureSample)(
+	[SharedThis = AsShared(), RDGBufferDesc](FRHICommandListImmediate& CommandList)
+	{
+		SharedThis->GPUBuffer = AllocatePooledBuffer(RDGBufferDesc, TEXT("RmaxInput Buffer"));
 
-FMediaTimeStamp FRivermaxMediaTextureSample::GetTime() const
-{
-	return FMediaTimeStamp(Time);
+	});
 }
 
 bool FRivermaxMediaTextureSample::IsCacheable() const
@@ -141,26 +78,9 @@ bool FRivermaxMediaTextureSample::IsCacheable() const
 	return false;
 }
 
-FRHITexture* FRivermaxMediaTextureSample::GetTexture() const
+uint8* FRivermaxMediaTextureSample::GetVideoBufferRawPtr(uint32 VideoBufferSize)
 {
-	return nullptr;
-}
-
-void* FRivermaxMediaTextureSample::RequestBuffer(uint32 InBufferSize)
-{
-	Buffer.Reset();
-	Buffer.SetNumUninitialized(InBufferSize);
-	return Buffer.GetData();
-}
-
-void FRivermaxMediaTextureSample::SetBuffer(TRefCountPtr<FRDGPooledBuffer> NewBuffer)
-{
-	GPUBuffer = NewBuffer;
-}
-
-TSharedPtr<FRivermaxMediaPlayer> FRivermaxMediaTextureSample::GetPlayer() const
-{
-	return WeakPlayer.Pin();
+	return reinterpret_cast<uint8*>(GetOrRequestBuffer(VideoBufferSize));
 }
 
 ERivermaxMediaSourcePixelFormat FRivermaxMediaTextureSample::GetInputFormat() const
@@ -168,34 +88,43 @@ ERivermaxMediaSourcePixelFormat FRivermaxMediaTextureSample::GetInputFormat() co
 	return InputFormat;
 }
 
-bool FRivermaxMediaTextureSample::NeedsSRGBToLinearConversion() const
+void FRivermaxMediaTextureSample::SetInputFormat(ERivermaxMediaSourcePixelFormat InFormat)
 {
-	return bIsInputSRGB;
-}
+	InputFormat = InFormat;
 
-bool FRivermaxMediaTextureSamples::FetchVideo(TRange<FTimespan> TimeRange, TSharedPtr<IMediaTextureSample, ESPMode::ThreadSafe>& OutSample)
-{
-	if (!CurrentSample.IsValid())
+	switch (InputFormat)
 	{
-		return false;
+		case ERivermaxMediaSourcePixelFormat::RGB_12bit:
+			// Falls through
+		case ERivermaxMediaSourcePixelFormat::RGB_16bit_Float:
+		{
+			SampleFormat = EMediaTextureSampleFormat::FloatRGBA;
+			break;
+		}
+		case ERivermaxMediaSourcePixelFormat::RGB_10bit:
+			// Falls through
+		case ERivermaxMediaSourcePixelFormat::YUV422_10bit:
+		{
+			SampleFormat = EMediaTextureSampleFormat::CharBGR10A2;
+			break;
+		}
+		case ERivermaxMediaSourcePixelFormat::YUV422_8bit:
+			// Falls through
+		case ERivermaxMediaSourcePixelFormat::RGB_8bit:
+			// Falls through
+		default:
+		{
+			SampleFormat = EMediaTextureSampleFormat::CharBGRA;
+			break;
+		}
 	}
-
-	OutSample = CurrentSample;
-
-	CurrentSample.Reset();
-
-	return true;
 }
 
-bool FRivermaxMediaTextureSamples::PeekVideoSampleTime(FMediaTimeStamp& TimeStamp)
+TRefCountPtr<FRDGPooledBuffer> FRivermaxMediaTextureSample::GetGPUBuffer() const
 {
-	return false;
+	return GPUBuffer;
 }
 
-void FRivermaxMediaTextureSamples::FlushSamples()
-{
-	CurrentSample.Reset();
-}
 
 }
 
diff --git a/Engine/Plugins/VirtualProduction/Rivermax/RivermaxMedia/Source/RivermaxMedia/Private/RivermaxMediaTextureSample.h b/Engine/Plugins/VirtualProduction/Rivermax/RivermaxMedia/Source/RivermaxMedia/Private/RivermaxMediaTextureSample.h
--- a/Engine/Plugins/VirtualProduction/Rivermax/RivermaxMedia/Source/RivermaxMedia/Private/RivermaxMediaTextureSample.h
+++ b/Engine/Plugins/VirtualProduction/Rivermax/RivermaxMedia/Source/RivermaxMedia/Private/RivermaxMediaTextureSample.h
@@ -4,153 +4,140 @@
 
 #include "IMediaTextureSample.h"
 
-#include "IMediaSamples.h"
+#include "IRivermaxInputStream.h"
+#include "MediaIOCoreTextureSampleBase.h"
+#include "MediaIOCoreSamples.h"
 #include "MediaShaders.h"
 #include "RivermaxMediaSource.h"
-#include "Templates/PimplPtr.h"
 #include "Templates/SharedPointer.h"
 
-
+class FEvent;
 
 namespace UE::RivermaxMedia
 {
-
-class FRivermaxMediaTextureSampleConverter;
-class FRivermaxMediaPlayer;
-
-/** Used to setup the configuration of a sample */
-struct FSampleConfigurationArgs
-{
-	/** Player that produced the sample */
-	TSharedPtr<FRivermaxMediaPlayer> Player;
-
-	/** Width of the sample in pixels */
-	uint32 Width = 0;
-
-	/** Height of the sample in pixels */
-	uint32 Height = 0;
-
-	/** Pixel format of the sample */
-	ERivermaxMediaSourcePixelFormat SampleFormat = ERivermaxMediaSourcePixelFormat::RGB_10bit;
-
-	/** Timestamp of the sample */
-	FTimespan Time = 0;
-
-	/** Frame rate at which samples are created */
-	FFrameRate FrameRate = {24, 1};
-
-	/** Whether sample is in SRGB space and needs to be converted to linear */
-	bool bInIsSRGBInput = false;
-};
-
-/**
- * Implements a media texture sample for RivermaxMediaPlayer.  
- */
-class FRivermaxMediaTextureSample : public IMediaTextureSample, public TSharedFromThis<FRivermaxMediaTextureSample>
-{
-public:
-
-	FRivermaxMediaTextureSample();
-
-	//~ Begin IMediaTextureSample interface
-	const void* GetBuffer() override;
-	virtual FIntPoint GetDim() const override;
-	virtual FTimespan GetDuration() const override;
-	virtual EMediaTextureSampleFormat GetFormat() const override;
-	virtual FIntPoint GetOutputDim() const override;
-	virtual uint32 GetStride() const override;
-	virtual FMediaTimeStamp GetTime() const override;
-	virtual bool IsCacheable() const override;
-	virtual const FMatrix& GetYUVToRGBMatrix() const override;
-	virtual bool IsOutputSrgb() const override;
-
-#if WITH_ENGINE
-	virtual IMediaTextureSampleConverter* GetMediaTextureSampleConverter() override;
-	virtual FRHITexture* GetTexture() const override;
-#endif
-	//~ End IMediaTextureSample interface
-
-	bool ConfigureSample(const FSampleConfigurationArgs& Args);
-
-	/** Initialized a RDG buffer based on the description required. Only useful for gpudirect functionality */
-	void InitializeGPUBuffer(const FIntPoint& InResolution, ERivermaxMediaSourcePixelFormat InSampleFormat);
-	
-	/** Returns RDG allocated buffer */
-	TRefCountPtr<FRDGPooledBuffer> GetGPUBuffer() const;
-
-	void SetBuffer(TRefCountPtr<FRDGPooledBuffer> NewBuffer);
-
-	/** Returns an uninitialized buffer of size InBufferSize */
-	void* RequestBuffer(uint32 InBufferSize);
-
-	/** Returns the player that created this sample */
-	TSharedPtr<FRivermaxMediaPlayer> GetPlayer() const;
-
-	/** Returns the incoming sample format */
-	ERivermaxMediaSourcePixelFormat GetInputFormat() const;
-
-	/** Whether this samples is coming as srgb and needs to be converted to linear during buffer conversion */
-	bool NeedsSRGBToLinearConversion() const;
-
-private:
-
-	/** Texture converted used to handle incoming 2110 formats and convert them to RGB textures the engine handles */
-	TPimplPtr<FRivermaxMediaTextureSampleConverter> Converter;
-	
-	/** Pooled buffer used for gpudirect functionality. Received content will already be on GPU when received from NIC */
-	TRefCountPtr<FRDGPooledBuffer> GPUBuffer;
-
-	/** Player that created this sample */
-	TWeakPtr<FRivermaxMediaPlayer> WeakPlayer;
-
-	/** Player time for this sample to be evaluated at */
-	FTimespan Time;
-
-	/** Texture stride */
-	uint32 Stride = 0;
-
-	/** Texture dimensions */
-	FIntPoint Dimension = FIntPoint::ZeroValue;
-
-	/** Sample format. */
-	EMediaTextureSampleFormat SampleFormat = EMediaTextureSampleFormat::CharBGRA;
-
-	/** Format in the rivermax realm */
-	ERivermaxMediaSourcePixelFormat InputFormat;
-
-	/** Duration for which the sample is valid. */
-	FTimespan Duration;
-
-	/** Whether the sample is in sRGB space and requires an explicit conversion to linear */
-	bool bIsInputSRGB = false;
-
-	/** System memory buffer to be used by converter when not using gpudirect */
-	TArray<uint8, TAlignedHeapAllocator<4096>> Buffer;
-};
-
-/** This is the mostly empty media samples that only carries a single sample per frame,
- *  that the player provides.
- */
-class FRivermaxMediaTextureSamples : public IMediaSamples
-{
-public:
-	FRivermaxMediaTextureSamples() = default;
-	FRivermaxMediaTextureSamples(const FRivermaxMediaTextureSamples&) = delete;
-	FRivermaxMediaTextureSamples& operator=(const FRivermaxMediaTextureSamples&) = delete;
-
-public:
-
-	//~ Begin IMediaSamples interface
-
-	virtual bool FetchVideo(TRange<FTimespan> TimeRange, TSharedPtr<IMediaTextureSample, ESPMode::ThreadSafe>& OutSample) override;
-	virtual bool PeekVideoSampleTime(FMediaTimeStamp& TimeStamp) override;
-	virtual void FlushSamples() override;
-
-	//~ End IMediaSamples interface
-
-public:
-
-	TSharedPtr<FRivermaxMediaTextureSample, ESPMode::ThreadSafe> CurrentSample;
-};
-
+	/**
+	 * Implements a media texture sample for RivermaxMediaPlayer.  
+	 */
+	class FRivermaxMediaTextureSample : public FMediaIOCoreTextureSampleBase, public UE::RivermaxCore::IRivermaxVideoSample, public TSharedFromThis<FRivermaxMediaTextureSample>
+	{
+	public:
+		FRivermaxMediaTextureSample();
+		virtual ~FRivermaxMediaTextureSample() override;
+
+		//~ Begin IMediaTextureSample interface
+		virtual bool IsCacheable() const override;
+		virtual const FMatrix& GetYUVToRGBMatrix() const override;
+
+	#if WITH_ENGINE
+		virtual IMediaTextureSampleConverter* GetMediaTextureSampleConverter() override;
+	#endif
+		//~ End IMediaTextureSample interface
+
+		//~ Begin FMediaIOCoreTextureSampleBase interface
+		virtual bool InitializeJITR(const FMediaIOCoreSampleJITRConfigurationArgs& Args) override;
+		virtual void CopyConfiguration(const TSharedPtr<FMediaIOCoreTextureSampleBase>& SourceSample) override;
+		//~ End FMediaIOCoreTextureSampleBase interface
+
+		/** Initialized a RDG buffer based on the description required. Only useful for gpudirect functionality */
+		void InitializeGPUBuffer(const FIntPoint& InResolution, ERivermaxMediaSourcePixelFormat InSampleFormat, bool bSupportsGPUDirect);
+
+		/** Returns the incoming sample format */
+		ERivermaxMediaSourcePixelFormat GetInputFormat() const;
+		void SetInputFormat(ERivermaxMediaSourcePixelFormat InFormat);
+
+		// Inherited via IRivermaxVideoSample
+		virtual TRefCountPtr<FRDGPooledBuffer> GetGPUBuffer() const override;
+		virtual uint8* GetVideoBufferRawPtr(uint32 VideoBufferSize) override;
+
+		virtual void InitializePoolable() override {};
+		virtual void ShutdownPoolable() override 
+		{
+			// This means that the sample is being held and managed by the Player
+			FScopeLock Lock(&StateChangeCriticalSecion);
+
+			// When this sample is returned back to the pool, it means that it is done rendering and is released from sample container.
+			SetAwaitingForGPUTransfer(false);
+			SetReceptionState(IRivermaxSample::ESampleState::ReadyForReception);
+			if (SampleConversionFence.IsValid())
+			{
+				SampleConversionFence->Clear();
+			}
+
+			LockedMemory = nullptr;
+			MarkRenderingComplete();
+		};
+
+		virtual bool IsReadyForReuse() override
+		{
+			return !IsBeingRendered();
+		}
+
+		/** 
+		* Attempts to lock this sample for rendering. 
+		* Returns true if sample is ok to be rendered. 
+		* Returns false if sample is already being rendered. 
+		*/
+		bool TryLockForRendering()
+		{
+			FScopeLock Lock(&StateChangeCriticalSecion);
+			if (bIsPendingRendering)
+			{
+				// sample is already being rendered.
+				return false;
+			}
+			else
+			{
+				bIsPendingRendering = true;
+				return true;
+			}
+		};
+
+		/** 
+		* Marks that this sample can be rendered again if need be.
+		*/
+		void MarkRenderingComplete()
+		{
+			FScopeLock Lock(&StateChangeCriticalSecion);
+			bIsPendingRendering = false;
+		};
+
+		bool IsBeingRendered()
+		{
+			FScopeLock Lock(&StateChangeCriticalSecion);
+			return bIsPendingRendering;
+		};
+
+		FEvent* GetSampleReceivedEvent() 
+		{
+			return SampleReceivedEvent;
+		}
+
+	public:
+		/** Locked memory of gpu buffer when uploading */
+		void* LockedMemory = nullptr;
+
+		/** Write fence enqueued after sample conversion to know when it's ready to be reused */
+		FGPUFenceRHIRef SampleConversionFence;
+
+	private:
+		/** This event isued to wait for the sample to be fully received.*/
+		FEvent* SampleReceivedEvent = nullptr;
+
+		/** True when queued for rendering. Will be false once fence has been written, after shader usage. */
+		std::atomic<bool> bIsPendingRendering = false;
+
+	private:
+
+		/** Format in the rivermax realm */
+		ERivermaxMediaSourcePixelFormat InputFormat = ERivermaxMediaSourcePixelFormat::YUV422_8bit;
+
+		/** Texture stride */
+		uint32 Stride = 0;
+
+		/** Pooled buffer used for gpudirect functionality. Received content will already be on GPU when received from NIC */
+		TRefCountPtr<FRDGPooledBuffer> GPUBuffer;
+	};
+
+	class FRivermaxMediaTextureSamplePool : public TMediaObjectPool<FRivermaxMediaTextureSample> { };
 }
+
diff --git a/Engine/Plugins/VirtualProduction/Rivermax/RivermaxMedia/Source/RivermaxMedia/Private/RivermaxMediaTextureSampleConverter.cpp b/Engine/Plugins/VirtualProduction/Rivermax/RivermaxMedia/Source/RivermaxMedia/Private/RivermaxMediaTextureSampleConverter.cpp
--- a/Engine/Plugins/VirtualProduction/Rivermax/RivermaxMedia/Source/RivermaxMedia/Private/RivermaxMediaTextureSampleConverter.cpp
+++ b/Engine/Plugins/VirtualProduction/Rivermax/RivermaxMedia/Source/RivermaxMedia/Private/RivermaxMediaTextureSampleConverter.cpp
@@ -2,239 +2,17 @@
 
 #include "RivermaxMediaTextureSampleConverter.h"
 
-#include "RenderGraphBuilder.h"
-#include "RivermaxMediaLog.h"
-#include "RivermaxMediaPlayer.h"
-#include "RivermaxMediaTextureSample.h"
-#include "RivermaxMediaUtils.h"
-#include "RivermaxShaders.h"
-
-DECLARE_GPU_STAT(RivermaxSource_SampleConversion);
-
 namespace UE::RivermaxMedia
 {
-
-void FRivermaxMediaTextureSampleConverter::Setup(const TSharedPtr<FRivermaxMediaTextureSample>& InSample)
-{
-	Sample = InSample;
-}
-
-bool FRivermaxMediaTextureSampleConverter::Convert(FRHICommandListImmediate& RHICmdList, FTextureRHIRef& InDestinationTexture, const FConversionHints& Hints)
-{
-	TRACE_CPUPROFILER_EVENT_SCOPE(RivermaxSampleConverter::Convert);
-
-	using namespace UE::RivermaxShaders;
-	using namespace UE::RivermaxMediaUtils::Private;
-
-	TSharedPtr<FRivermaxMediaTextureSample> SamplePtr = Sample.Pin();
-	if (SamplePtr.IsValid() == false)
-	{
-		return false;
-	}
-
-	TSharedPtr<FRivermaxMediaPlayer> PlayerPtr = SamplePtr->GetPlayer();
-	if (PlayerPtr.IsValid() == false)
-	{
-		return false;
-	}
-
-	// Get sample setup from our player
-	FSampleConverterOperationSetup Setup;
-	const bool bCanRender = PlayerPtr->LateUpdateSetupSample(Setup);
-	if (bCanRender == false)
-	{
-		return false;
-	}
-
-	// If no input data was provided, no need to render
-	if (Setup.GetGPUBufferFunc == nullptr && Setup.GetSystemBufferFunc == nullptr)
-	{
-		ensureMsgf(false, TEXT("Rivermax player late update succeeded but didn't provide any source data."));
-		return false;
-	}
-	
-	FRDGBuilder GraphBuilder(RHICmdList);
-	if (Setup.PreConvertFunc)
+	uint32 FRivermaxMediaTextureSampleConverter::GetConverterInfoFlags() const
 	{
-		Setup.PreConvertFunc(GraphBuilder);
-	}
-
-	const FSourceBufferDesc SourceBufferDesc = GetBufferDescription(SamplePtr->GetDim(), SamplePtr->GetInputFormat());
-	{
-		FRDGBufferRef InputBuffer;
-
-		RDG_EVENT_SCOPE_STAT(GraphBuilder, RivermaxSource_SampleConversion, "Rivermax_SampleConverter");
-		RDG_GPU_STAT_SCOPE(GraphBuilder, RivermaxSource_SampleConversion)
-
-		FRDGTextureRef OutputResource = GraphBuilder.RegisterExternalTexture(CreateRenderTarget(InDestinationTexture, TEXT("RivermaxMediaTextureOutputResource")));
-
-		// If we have a valid GPUBuffer, i.e GPUDirect is involved, use that one. Otherwise, take the system buffer and upload it in a new structured buffer.
-		if (Setup.GetGPUBufferFunc)
-		{
-			InputBuffer = GraphBuilder.RegisterExternalBuffer(Setup.GetGPUBufferFunc(), TEXT("RMaxGPUBuffer"));
-		}
-		else if (Setup.GetSystemBufferFunc)
+		if (TSharedPtr<FMediaIOCoreTextureSampleBase> ProxySampleSharedPtr = JITRProxySample.Pin())
 		{
-			TRACE_CPUPROFILER_EVENT_SCOPE(RivermaxSampleConverter::CreateStructuredBuffer);
-			FRDGBufferNumElementsCallback NumElementCallback = [NumElements = SourceBufferDesc.NumberOfElements]()
-			{
-				return NumElements;
-			};
-
-			FRDGBufferInitialDataCallback GetDataCallback = Setup.GetSystemBufferFunc;
-			FRDGBufferInitialDataSizeCallback TotalSizeCallback = [TotalSize = SourceBufferDesc.BytesPerElement * SourceBufferDesc.NumberOfElements]()
+			if (ProxySampleSharedPtr->GetEncodingType() != UE::Color::EEncoding::None || ProxySampleSharedPtr->GetColorSpaceType() != UE::Color::EColorSpace::None)
 			{
-				return TotalSize;
-			};
-
-			InputBuffer = CreateStructuredBuffer(GraphBuilder, TEXT("RivermaxInputBuffer"), SourceBufferDesc.BytesPerElement, MoveTemp(NumElementCallback), MoveTemp(GetDataCallback), MoveTemp(TotalSizeCallback));
-		}
-		else
-		{
-			return false;
-		}
-
-		
-		const FIntPoint ProcessedOutputDimension = { (int32)SourceBufferDesc.NumberOfElements, 1 };
-		const FIntVector GroupCount = FComputeShaderUtils::GetGroupCount(SourceBufferDesc.NumberOfElements, 64);
-		FGlobalShaderMap* GlobalShaderMap = GetGlobalShaderMap(GMaxRHIFeatureLevel);
-
-		//Configure shader and add conversion pass based on desired pixel format
-		switch (SamplePtr->GetInputFormat())
-		{
-		case ERivermaxMediaSourcePixelFormat::YUV422_8bit:
-		{
-			TRACE_CPUPROFILER_EVENT_SCOPE(RivermaxSampleConverter::YUV8ShaderSetup);
-
-			FYUV8Bit422ToRGBACS::FPermutationDomain PermutationVector;
-			PermutationVector.Set<FYUV8Bit422ToRGBACS::FSRGBToLinear>(SamplePtr->NeedsSRGBToLinearConversion());
-
-			const FMatrix YUVToRGBMatrix = SamplePtr->GetYUVToRGBMatrix();
-			const FVector YUVOffset(MediaShaders::YUVOffset8bits);
-			TShaderMapRef<FYUV8Bit422ToRGBACS> ComputeShader(GlobalShaderMap, PermutationVector);
-			FYUV8Bit422ToRGBACS::FParameters* Parameters = ComputeShader->AllocateAndSetParameters(GraphBuilder, InputBuffer, OutputResource, YUVToRGBMatrix, YUVOffset, ProcessedOutputDimension.X, ProcessedOutputDimension.Y);
-
-			FComputeShaderUtils::AddPass(
-				GraphBuilder
-				, RDG_EVENT_NAME("YUV8Bit422ToRGBA")
-				, ComputeShader
-				, Parameters
-				, GroupCount);
-			break;
-		}
-		case ERivermaxMediaSourcePixelFormat::YUV422_10bit:
-		{
-			TRACE_CPUPROFILER_EVENT_SCOPE(RivermaxSampleConverter::YUV10ShaderSetup);
-			
-			FYUV10Bit422ToRGBACS::FPermutationDomain PermutationVector;
-			PermutationVector.Set<FYUV10Bit422ToRGBACS::FSRGBToLinear>(SamplePtr->NeedsSRGBToLinearConversion());
-
-			const FMatrix YUVToRGBMatrix = SamplePtr->GetYUVToRGBMatrix();
-			const FVector YUVOffset(MediaShaders::YUVOffset10bits);
-			TShaderMapRef<FYUV10Bit422ToRGBACS> ComputeShader(GlobalShaderMap, PermutationVector);
-			FYUV10Bit422ToRGBACS::FParameters* Parameters = ComputeShader->AllocateAndSetParameters(GraphBuilder, InputBuffer, OutputResource, YUVToRGBMatrix, YUVOffset, ProcessedOutputDimension.X, ProcessedOutputDimension.Y);
-
-			FComputeShaderUtils::AddPass(
-				GraphBuilder
-				, RDG_EVENT_NAME("YUV10Bit422ToRGBA")
-				, ComputeShader
-				, Parameters
-				, GroupCount);
-			break;
-		}
-		case ERivermaxMediaSourcePixelFormat::RGB_8bit:
-		{
-			TRACE_CPUPROFILER_EVENT_SCOPE(RivermaxSampleConverter::RGB8ShaderSetup);
-
-			FRGB8BitToRGBA8CS::FPermutationDomain PermutationVector;
-			PermutationVector.Set<FRGB8BitToRGBA8CS::FSRGBToLinear>(SamplePtr->NeedsSRGBToLinearConversion());
-
-			TShaderMapRef<FRGB8BitToRGBA8CS> ComputeShader(GlobalShaderMap, PermutationVector);
-			FRGB8BitToRGBA8CS::FParameters* Parameters = ComputeShader->AllocateAndSetParameters(GraphBuilder, InputBuffer, OutputResource, ProcessedOutputDimension.X, ProcessedOutputDimension.Y);
-
-			FComputeShaderUtils::AddPass(
-				GraphBuilder
-				, RDG_EVENT_NAME("RGB8BitToRGBA8")
-				, ComputeShader
-				, Parameters
-				, GroupCount);
-			break;
-		}
-		case ERivermaxMediaSourcePixelFormat::RGB_10bit:
-		{
-			TRACE_CPUPROFILER_EVENT_SCOPE(RivermaxSampleConverter::RGB10ShaderSetup);
-
-			FRGB10BitToRGBA10CS::FPermutationDomain PermutationVector;
-			PermutationVector.Set<FRGB10BitToRGBA10CS::FSRGBToLinear>(SamplePtr->NeedsSRGBToLinearConversion());
-
-			TShaderMapRef<FRGB10BitToRGBA10CS> ComputeShader(GlobalShaderMap, PermutationVector);
-			FRGB10BitToRGBA10CS::FParameters* Parameters = ComputeShader->AllocateAndSetParameters(GraphBuilder, InputBuffer, OutputResource, ProcessedOutputDimension.X, ProcessedOutputDimension.Y);
-
-			FComputeShaderUtils::AddPass(
-				GraphBuilder
-				, RDG_EVENT_NAME("RGB10BitToRGBA")
-				, ComputeShader
-				, Parameters
-				, GroupCount);
-			break;
-		}
-		case ERivermaxMediaSourcePixelFormat::RGB_12bit:
-		{
-			TRACE_CPUPROFILER_EVENT_SCOPE(RivermaxSampleConverter::RGB12ShaderSetup);
-
-			FRGB12BitToRGBA12CS::FPermutationDomain PermutationVector;
-			PermutationVector.Set<FRGB12BitToRGBA12CS::FSRGBToLinear>(SamplePtr->NeedsSRGBToLinearConversion());
-
-			TShaderMapRef<FRGB12BitToRGBA12CS> ComputeShader(GlobalShaderMap, PermutationVector);
-			FRGB12BitToRGBA12CS::FParameters* Parameters = ComputeShader->AllocateAndSetParameters(GraphBuilder, InputBuffer, OutputResource, ProcessedOutputDimension.X, ProcessedOutputDimension.Y);
-
-			FComputeShaderUtils::AddPass(
-				GraphBuilder
-				, RDG_EVENT_NAME("RGB12BitToRGBA")
-				, ComputeShader
-				, Parameters
-				, GroupCount);
-			break;
-		}
-		case ERivermaxMediaSourcePixelFormat::RGB_16bit_Float:
-		{
-			TRACE_CPUPROFILER_EVENT_SCOPE(RivermaxSampleConverter::RGB16FloatShaderSetup);
-
-			FRGB16fBitToRGBA16fCS::FPermutationDomain PermutationVector;
-			PermutationVector.Set<FRGB16fBitToRGBA16fCS::FSRGBToLinear>(SamplePtr->NeedsSRGBToLinearConversion());
-
-			TShaderMapRef<FRGB16fBitToRGBA16fCS> ComputeShader(GlobalShaderMap, PermutationVector);
-			FRGB16fBitToRGBA16fCS::FParameters* Parameters = ComputeShader->AllocateAndSetParameters(GraphBuilder, InputBuffer, OutputResource, ProcessedOutputDimension.X, ProcessedOutputDimension.Y);
-
-			FComputeShaderUtils::AddPass(
-				GraphBuilder
-				, RDG_EVENT_NAME("RGB16fBitToRGBA")
-				, ComputeShader
-				, Parameters
-				, GroupCount);
-			break;
-		}
-		default:
-		{
-			ensureMsgf(false, TEXT("Unhandled pixel format (%d) given to Rivermax MediaSample converter"), SamplePtr->GetInputFormat());
-			return false;
-		}
+				return ConverterInfoFlags_WillCreateOutputTexture | ConverterInfoFlags_PreprocessOnly;
+			}
 		}
+		return ConverterInfoFlags_WillCreateOutputTexture;
 	}
-
-	if (Setup.PostConvertFunc)
-	{
-		Setup.PostConvertFunc(GraphBuilder);
-	}
-
-	GraphBuilder.Execute();
-
-	return true;
-}
-
-uint32 FRivermaxMediaTextureSampleConverter::GetConverterInfoFlags() const
-{
-	return IMediaTextureSampleConverter::ConverterInfoFlags_NeedUAVOutputTexture;
-}
-
-}
+}
\ No newline at end of file
diff --git a/Engine/Plugins/VirtualProduction/Rivermax/RivermaxMedia/Source/RivermaxMedia/Private/RivermaxMediaTextureSampleConverter.h b/Engine/Plugins/VirtualProduction/Rivermax/RivermaxMedia/Source/RivermaxMedia/Private/RivermaxMediaTextureSampleConverter.h
--- a/Engine/Plugins/VirtualProduction/Rivermax/RivermaxMedia/Source/RivermaxMedia/Private/RivermaxMediaTextureSampleConverter.h
+++ b/Engine/Plugins/VirtualProduction/Rivermax/RivermaxMedia/Source/RivermaxMedia/Private/RivermaxMediaTextureSampleConverter.h
@@ -2,7 +2,7 @@
 
 #pragma once
 
-#include "IMediaTextureSampleConverter.h"
+#include "MediaIOCoreTextureSampleConverter.h"
 
 #include "RivermaxMediaSource.h"
 
@@ -12,58 +12,34 @@ class FRDGBuilder;
 
 namespace UE::RivermaxMedia
 {
-
-class FRivermaxMediaTextureSample;
-class FRivermaxMediaPlayer;
-
-using FPreInputConvertFunc = TUniqueFunction<void(FRDGBuilder& GraphBuilder)>;
-using FGetSystemBufferFunc = TFunction<const void *()>;
-using FGetGPUBufferFunc = TFunction<TRefCountPtr<FRDGPooledBuffer>()>;
-using FPostInputConvertFunc = TUniqueFunction<void(FRDGBuilder& GraphBuilder)>;
-
-/** Structure used during late update to let player configure some operations */
-struct FSampleConverterOperationSetup
-{
-
-
-	/** Function to be called before setting up the sample conversion graph */
-	FPreInputConvertFunc PreConvertFunc = nullptr;
-
-	/** Function used to retrieve which system buffer to use. Can block until data is available. */
-	FGetSystemBufferFunc GetSystemBufferFunc = nullptr;
-	
-	/** Function used to retrieve gpu buffer if available */
-	FGetGPUBufferFunc GetGPUBufferFunc = nullptr;
-
-	/** Function to be called after setting up the sample conversion graph */
-	FPostInputConvertFunc PostConvertFunc = nullptr;
-};
-
-
-/**
- * Sample converter used for 2110 video samples. Supports some of the pixel formats
- * from 2110-20. Data is expected to be packed in a buffer and will be converted to 
- * an output texture to be rendered.
- */
-class FRivermaxMediaTextureSampleConverter : public IMediaTextureSampleConverter
-{
-public:
-
-	FRivermaxMediaTextureSampleConverter() = default;
-	~FRivermaxMediaTextureSampleConverter() = default;
-
-	/** Configures settings to convert incoming sample */
-	void Setup(const TSharedPtr<FRivermaxMediaTextureSample>& InSample);
-
-	//~ Begin IMediaTextureSampleConverter interface
-	virtual bool Convert(FRHICommandListImmediate& RHICmdList, FTextureRHIRef& InDstTexture, const FConversionHints& Hints) override;
-	virtual uint32 GetConverterInfoFlags() const override;
-	//~ End IMediaTextureSampleConverter interface
-
-private:
-
-	/** Sample holding the buffer to convert */
-	TWeakPtr<FRivermaxMediaTextureSample> Sample;
-};
+	class FRivermaxMediaTextureSample;
+	class FRivermaxMediaPlayer;
+
+	using FPreInputConvertFunc = TUniqueFunction<void(FRDGBuilder& GraphBuilder)>;
+	using FGetSystemBufferFunc = TFunction<const void* ()>;
+	using FGetGPUBufferFunc = TFunction<TRefCountPtr<FRDGPooledBuffer>()>;
+	using FPostInputConvertFunc = TUniqueFunction<void(FRDGBuilder& GraphBuilder)>;
+
+	/** Structure used during late update to let player configure some operations */
+	struct FSampleConverterOperationSetup
+	{
+		/** Function to be called before setting up the sample conversion graph */
+		FPreInputConvertFunc PreConvertFunc = nullptr;
+
+		/** Function used to retrieve which system buffer to use. Can block until data is available. */
+		FGetSystemBufferFunc GetSystemBufferFunc = nullptr;
+
+		/** Function used to retrieve gpu buffer if available */
+		FGetGPUBufferFunc GetGPUBufferFunc = nullptr;
+
+		/** Function to be called after setting up the sample conversion graph */
+		FPostInputConvertFunc PostConvertFunc = nullptr;
+	};
+
+	class FRivermaxMediaTextureSampleConverter : public FMediaIOCoreTextureSampleConverter
+	{
+	public:
+		virtual uint32 GetConverterInfoFlags() const override;
+	};
 
 }
diff --git a/Engine/Plugins/VirtualProduction/Rivermax/RivermaxMedia/Source/RivermaxMedia/Private/RivermaxMediaUtils.cpp b/Engine/Plugins/VirtualProduction/Rivermax/RivermaxMedia/Source/RivermaxMedia/Private/RivermaxMediaUtils.cpp
--- a/Engine/Plugins/VirtualProduction/Rivermax/RivermaxMedia/Source/RivermaxMedia/Private/RivermaxMediaUtils.cpp
+++ b/Engine/Plugins/VirtualProduction/Rivermax/RivermaxMedia/Source/RivermaxMedia/Private/RivermaxMediaUtils.cpp
@@ -3,6 +3,14 @@
 #include "RivermaxMediaUtils.h"
 
 #include "RivermaxShaders.h"
+#include "Serialization/CustomVersion.h"
+
+// Guid generated by FGuid::NewGuid
+const FGuid FRivermaxMediaVersion::GUID(0x867597A5, 0x492791B6, 0x3FDACDB0, 0x60344667);
+
+// RivermaxMedia Version registration with core.
+FCustomVersionRegistration GRegisterCustomizableObjectCustomVersion(FRivermaxMediaVersion::GUID, FRivermaxMediaVersion::LatestVersion, TEXT("RivermaxMediaVersion"));
+
 
 namespace UE::RivermaxMediaUtils::Private
 {
diff --git a/Engine/Plugins/VirtualProduction/Rivermax/RivermaxMedia/Source/RivermaxMedia/Private/RivermaxMediaUtils.h b/Engine/Plugins/VirtualProduction/Rivermax/RivermaxMedia/Source/RivermaxMedia/Private/RivermaxMediaUtils.h
--- a/Engine/Plugins/VirtualProduction/Rivermax/RivermaxMedia/Source/RivermaxMedia/Private/RivermaxMediaUtils.h
+++ b/Engine/Plugins/VirtualProduction/Rivermax/RivermaxMedia/Source/RivermaxMedia/Private/RivermaxMediaUtils.h
@@ -9,6 +9,8 @@
 #include "RivermaxMediaSource.h"
 #include "RivermaxTypes.h"
 
+struct FGuid;
+
 namespace UE::RivermaxMediaUtils::Private
 {
 	struct FSourceBufferDesc
@@ -25,4 +27,25 @@ namespace UE::RivermaxMediaUtils::Private
 	UE::RivermaxCore::EFrameLockingMode MediaOutputFrameLockingToRivermax(ERivermaxFrameLockingMode InFrameLockingMode);
 	FSourceBufferDesc GetBufferDescription(const FIntPoint& Resolution, ERivermaxMediaSourcePixelFormat InPixelFormat);
 
-}
\ No newline at end of file
+}
+
+// Custom version to keep track of and restore depricated properties.
+struct RIVERMAXMEDIA_API FRivermaxMediaVersion
+{
+	enum Type
+	{
+		BeforeCustomVersionAdded = 0,
+
+		// Add new versions above this comment.
+		VersionPlusOne,
+		LatestVersion = VersionPlusOne - 1
+
+
+	};
+
+	// Rivermax Guild
+	const static FGuid GUID;
+
+private:
+	FRivermaxMediaVersion() {}
+};
\ No newline at end of file
diff --git a/Engine/Plugins/VirtualProduction/Rivermax/RivermaxMedia/Source/RivermaxMedia/Public/RivermaxMediaSource.h b/Engine/Plugins/VirtualProduction/Rivermax/RivermaxMedia/Source/RivermaxMedia/Public/RivermaxMediaSource.h
--- a/Engine/Plugins/VirtualProduction/Rivermax/RivermaxMedia/Source/RivermaxMedia/Public/RivermaxMediaSource.h
+++ b/Engine/Plugins/VirtualProduction/Rivermax/RivermaxMedia/Source/RivermaxMedia/Public/RivermaxMediaSource.h
@@ -2,7 +2,7 @@
 
 #pragma once
 
-#include "TimeSynchronizableMediaSource.h"
+#include "CaptureCardMediaSource.h"
 
 #include "MediaIOCoreDefinitions.h"
 #include "RivermaxTypes.h"
@@ -24,41 +24,55 @@ enum class ERivermaxMediaSourcePixelFormat : uint8
 };
 
 /**
- * Player mode to be used.
+ * Player mode to be used. Deprecated in UE5.5
  */
 UENUM()
-enum class ERivermaxPlayerMode : uint8
+enum class ERivermaxPlayerMode_DEPRECATED : uint8
 {
 	Latest, // Uses latest sample available
-	Framelock // Uses incoming samples frame number to match with local engine frame number. Should be used with nDisplay
+	Framelock, // Uses incoming samples frame number to match with local engine frame number. Should be used with nDisplay
 };
 
 /**
  * Media source for Rivermax streams.
  */
 UCLASS(BlueprintType, hideCategories=(Platforms,Object), meta=(MediaIOCustomLayout="Rivermax", DisplayName = "NVIDIA Rivermax Source"))
-class RIVERMAXMEDIA_API URivermaxMediaSource : public UTimeSynchronizableMediaSource
+class RIVERMAXMEDIA_API URivermaxMediaSource : public UCaptureCardMediaSource
 {
 	GENERATED_BODY()
 
 public:
 
+	URivermaxMediaSource();
+
+public:
+
+#if WITH_EDITORONLY_DATA
 	/**
 	 * Player mode to be used.
 	 * Latest : Default mode. Will use latest available at render time. No alignment.
 	 * Framelock : Will match sample's frame number with engine frame number. Meant to be used for UE-UE contexts like for nDisplay
 	 *           : Will wait for an expected to arrive before moving with render
 	 */
-	UPROPERTY(EditAnywhere, BlueprintReadWrite, Category = "Player")
-	ERivermaxPlayerMode PlayerMode = ERivermaxPlayerMode::Latest;
+	UE_DEPRECATED(5.5, "This property has been deprecated.")
+	UPROPERTY(meta = (DeprecatedProperty, DeprecationMessage = "Use Sample Evaluation Type and Framelock instead"))
+	ERivermaxPlayerMode_DEPRECATED PlayerMode_DEPRECATED = ERivermaxPlayerMode_DEPRECATED::Latest;
 
 	/** 
 	 * If true, when looking for the sample to render, the current frame number will be looked for.
 	 * If expected frame hasn't been received, waiting will occur.
 	 * If false, player will look for one frame behind
 	 */
-	UPROPERTY(EditAnywhere, BlueprintReadWrite, Category = "Player", meta = (EditCondition = "PlayerMode == ERivermaxPlayerMode::Framelock"))
-	bool bUseZeroLatency  = true;
+	UE_DEPRECATED(5.5, "This property has been deprecated.")
+	UPROPERTY(meta = (DeprecatedProperty, DeprecationMessage = "Use Frame Delay under Synchronization category"))
+	bool bUseZeroLatency_DEPRECATED = true;
+
+
+	/** Whether the video input is in sRGB color space.If true, sRGBToLinear will be done on incoming pixels before writing to media texture */
+	UE_DEPRECATED(5.5, "This property has been deprecated.")
+	UPROPERTY(meta = (DeprecatedProperty, DeprecationMessage = "Use Override Source Encoding instead"))
+	bool bIsSRGBInput_DEPRECATED = false;
+#endif
 
 	/** If false, use the default source buffer size. If true, a specific resolution will be used. */
 	UPROPERTY(BlueprintReadWrite, EditAnywhere, Category = "Media", meta = (InlineEditConditionToggle))
@@ -94,10 +108,6 @@ public:
 	UPROPERTY(EditAnywhere, BlueprintReadWrite, Category = "Options")
 	int32 Port = 50000;
 
-	/** Whether the video input is in sRGB color space.If true, sRGBToLinear will be done on incoming pixels before writing to media texture */
-	UPROPERTY(EditAnywhere, BlueprintReadWrite, Category = "Video")
-	bool bIsSRGBInput = false;
-
 	/** Whether to use GPUDirect if available (Memcopy from NIC to GPU directly bypassing system memory) if available */
 	UPROPERTY(EditAnywhere, BlueprintReadWrite, AdvancedDisplay, Category = "Video")
 	bool bUseGPUDirect = true;
@@ -110,6 +120,9 @@ public:
 	virtual bool HasMediaOption(const FName& Key) const override;
 	//~ End IMediaOptions interface
 
+	virtual void PostLoad() override;
+	void Serialize(FArchive& Ar) override;
+
 public:
 	//~ Begin UMediaSource interface
 	virtual FString GetUrl() const override;
diff --git a/Engine/Plugins/VirtualProduction/Rivermax/RivermaxMedia/Source/RivermaxMediaEditor/Private/ModularFeatures/RivermaxMediaInitializerFeature.cpp b/Engine/Plugins/VirtualProduction/Rivermax/RivermaxMedia/Source/RivermaxMediaEditor/Private/ModularFeatures/RivermaxMediaInitializerFeature.cpp
--- a/Engine/Plugins/VirtualProduction/Rivermax/RivermaxMedia/Source/RivermaxMediaEditor/Private/ModularFeatures/RivermaxMediaInitializerFeature.cpp
+++ b/Engine/Plugins/VirtualProduction/Rivermax/RivermaxMedia/Source/RivermaxMediaEditor/Private/ModularFeatures/RivermaxMediaInitializerFeature.cpp
@@ -51,8 +51,10 @@ void FRivermaxMediaInitializerFeature::InitializeMediaObjectForTile(UObject* Med
 {
 	if (URivermaxMediaSource* RivermaxMediaSource = Cast<URivermaxMediaSource>(MediaObject))
 	{
-		RivermaxMediaSource->PlayerMode          = ERivermaxPlayerMode::Framelock;
-		RivermaxMediaSource->bUseZeroLatency     = true;
+		RivermaxMediaSource->EvaluationType      = EMediaIOSampleEvaluationType::Timecode;
+		RivermaxMediaSource->bFramelock          = true;
+		RivermaxMediaSource->bUseTimeSynchronization = true;
+		RivermaxMediaSource->FrameDelay			 = 0;
 		RivermaxMediaSource->bOverrideResolution = false;
 		//RivermaxMediaSource->Resolution          = default value
 		RivermaxMediaSource->FrameRate           = { 60,1 };
@@ -60,7 +62,6 @@ void FRivermaxMediaInitializerFeature::InitializeMediaObjectForTile(UObject* Med
 		RivermaxMediaSource->InterfaceAddress    = GetRivermaxInterfaceAddress();
 		RivermaxMediaSource->StreamAddress       = GenerateStreamAddress(OnwerInfo.OwnerUniqueIdx, TilePos);
 		RivermaxMediaSource->Port                = 50000;
-		RivermaxMediaSource->bIsSRGBInput        = false;
 		RivermaxMediaSource->bUseGPUDirect       = true;
 	}
 	else if (URivermaxMediaOutput* RivermaxMediaOutput = Cast<URivermaxMediaOutput>(MediaObject))
@@ -85,16 +86,17 @@ void FRivermaxMediaInitializerFeature::InitializeMediaObjectForFullFrame(UObject
 {
 	if (URivermaxMediaSource* RivermaxMediaSource = Cast<URivermaxMediaSource>(MediaObject))
 	{
-		RivermaxMediaSource->PlayerMode          = ERivermaxPlayerMode::Framelock;
-		RivermaxMediaSource->bUseZeroLatency     = true;
+		RivermaxMediaSource->EvaluationType		 = EMediaIOSampleEvaluationType::Timecode;
+		RivermaxMediaSource->bFramelock			 = true;
+		RivermaxMediaSource->bUseTimeSynchronization = true;
 		RivermaxMediaSource->bOverrideResolution = false;
+		RivermaxMediaSource->FrameDelay = 0;
 		//RivermaxMediaSource->Resolution          = default value
 		RivermaxMediaSource->FrameRate           = { 60,1 };
 		RivermaxMediaSource->PixelFormat         = ERivermaxMediaSourcePixelFormat::RGB_10bit;
 		RivermaxMediaSource->InterfaceAddress    = GetRivermaxInterfaceAddress();
 		RivermaxMediaSource->StreamAddress       = GenerateStreamAddress(OnwerInfo.ClusterNodeUniqueIdx.Get(0), OnwerInfo.OwnerUniqueIdx, OnwerInfo.OwnerType);
 		RivermaxMediaSource->Port                = 50000;
-		RivermaxMediaSource->bIsSRGBInput        = false;
 		RivermaxMediaSource->bUseGPUDirect       = true;
 	}
 	else if (URivermaxMediaOutput* RivermaxMediaOutput = Cast<URivermaxMediaOutput>(MediaObject))
