Rivermax: Fixed a crash that happens if Media Profile uses rivermax timecode provider and Rivermax library isn't intialized for some reason.

#rnx

diff --git a/Engine/Plugins/VirtualProduction/Rivermax/RivermaxMedia/Source/RivermaxMedia/Private/RivermaxTimecodeProvider.cpp b/Engine/Plugins/VirtualProduction/Rivermax/RivermaxMedia/Source/RivermaxMedia/Private/RivermaxTimecodeProvider.cpp
--- a/Engine/Plugins/VirtualProduction/Rivermax/RivermaxMedia/Source/RivermaxMedia/Private/RivermaxTimecodeProvider.cpp
+++ b/Engine/Plugins/VirtualProduction/Rivermax/RivermaxMedia/Source/RivermaxMedia/Private/RivermaxTimecodeProvider.cpp
@@ -13,7 +13,7 @@
 bool URivermaxTimecodeProvider::FetchTimecode(FQualifiedFrameTime& OutFrameTime)
 {
 	IRivermaxCoreModule* RivermaxModule = FModuleManager::GetModulePtr<IRivermaxCoreModule>("RivermaxCore");
-	if (RivermaxModule && RivermaxModule->GetRivermaxManager())
+	if (RivermaxModule && RivermaxModule->GetRivermaxManager() && RivermaxModule->GetRivermaxManager()->IsLibraryInitialized())
 	{
 		// Get rivermax clock time, and truncate to timespan tick's resolution (100ns / tick)
 		const uint64 CurrentTime = RivermaxModule->GetRivermaxManager()->GetTime();
