Rivermax Media Input: Support 2110 media input sync using time of arrival.

Note: RTP Headers do not have PTP timestamps.

diff --git a/Engine/Plugins/VirtualProduction/Rivermax/RivermaxCore/Source/RivermaxCore/Private/RivermaxPTPUtils.cpp b/Engine/Plugins/VirtualProduction/Rivermax/RivermaxCore/Source/RivermaxCore/Private/RivermaxPTPUtils.cpp
--- a/Engine/Plugins/VirtualProduction/Rivermax/RivermaxCore/Source/RivermaxCore/Private/RivermaxPTPUtils.cpp
+++ b/Engine/Plugins/VirtualProduction/Rivermax/RivermaxCore/Source/RivermaxCore/Private/RivermaxPTPUtils.cpp
@@ -4,7 +4,6 @@
 
 #include "Misc/FrameRate.h"
 #include "RivermaxTypes.h"
-#include "RivermaxUtils.h"
 
 namespace UE::RivermaxCore
 {
@@ -65,9 +64,4 @@ namespace UE::RivermaxCore
 
 		return FSec + FNanoSec + FractionalFrames;
 	}
-
-	RIVERMAXCORE_API uint64 ConvertRTPTimeStampToSecondsTruncated(const uint64 InTimestamp)
-	{
-		return InTimestamp / UE::RivermaxCore::Private::Utils::MediaClockSampleRate;
-	}
 }
\ No newline at end of file
diff --git a/Engine/Plugins/VirtualProduction/Rivermax/RivermaxCore/Source/RivermaxCore/Public/RivermaxPTPUtils.h b/Engine/Plugins/VirtualProduction/Rivermax/RivermaxCore/Source/RivermaxCore/Public/RivermaxPTPUtils.h
--- a/Engine/Plugins/VirtualProduction/Rivermax/RivermaxCore/Source/RivermaxCore/Public/RivermaxPTPUtils.h
+++ b/Engine/Plugins/VirtualProduction/Rivermax/RivermaxCore/Source/RivermaxCore/Public/RivermaxPTPUtils.h
@@ -24,12 +24,6 @@ namespace UE::RivermaxCore
 	 * Returns current frame number for the incoming PTP time for the given frame rate.
 	 */
 	RIVERMAXCORE_API uint64 GetFrameNumber(const uint64 InPTPTimeNanosec, const FFrameRate& InRate);
-
-	/**
-	 * Converts an RTP timestamp in MediaClock period units to seconds
-	 * 2110-20 streams uses a standard media clock rate of 90kHz
-	 */
-	RIVERMAXCORE_API uint64 ConvertRTPTimeStampToSecondsTruncated(const uint64 InTimestamp);
 }
 
 
diff --git a/Engine/Plugins/VirtualProduction/Rivermax/RivermaxMedia/Source/RivermaxMedia/Private/RivermaxMediaPlayer.cpp b/Engine/Plugins/VirtualProduction/Rivermax/RivermaxMedia/Source/RivermaxMedia/Private/RivermaxMediaPlayer.cpp
--- a/Engine/Plugins/VirtualProduction/Rivermax/RivermaxMedia/Source/RivermaxMedia/Private/RivermaxMediaPlayer.cpp
+++ b/Engine/Plugins/VirtualProduction/Rivermax/RivermaxMedia/Source/RivermaxMedia/Private/RivermaxMediaPlayer.cpp
@@ -2,6 +2,7 @@
 
 #include "RivermaxMediaPlayer.h"
 
+#include "Engine/Engine.h"
 #include "IMediaEventSink.h"
 #include "IRivermaxCoreModule.h"
 #include "IRivermaxManager.h"
@@ -19,6 +20,7 @@
 #include "RivermaxPTPUtils.h"
 #include "RivermaxShaders.h"
 #include "RivermaxTracingUtils.h"
+#include "RivermaxTimecodeProvider.h"
 #include "RivermaxTypes.h"
 #include "Stats/Stats2.h"
 #include "Tasks/Task.h"
@@ -71,6 +73,37 @@ namespace UE::RivermaxMedia
 		return (GetFrameNumberWithAcountedLatency(InFrameDelay, FrameNumber)) % InMaxNumVideoFrameBuffer;
 	}
 
+	/** Returns current time. Adjusted to UTC and rolled over at 24 hours. */
+	FTimespan GetCurrentPTPTimeOfDay()
+	{
+		FTimespan CurrentTimespan;
+
+		const int64 NumberOfTicksPerDay = 60 * 60 * 24 * ETimespan::TicksPerSecond;
+
+		IRivermaxCoreModule* RivermaxModule = FModuleManager::GetModulePtr<IRivermaxCoreModule>("RivermaxCore");
+		if (RivermaxModule && RivermaxModule->GetRivermaxManager())
+		{
+			// Converting from nanoseconds to ticks.
+			CurrentTimespan = FTimespan(RivermaxModule->GetRivermaxManager()->GetTime() / ETimespan::NanosecondsPerTick);
+			UTimecodeProvider* Provider = GEngine->GetTimecodeProvider();
+			
+			// Convert from TAI PTP Time to UTC
+			if (Provider && Provider->GetName().Contains("RivermaxTimecodeProvider"))
+			{
+				URivermaxTimecodeProvider* RmaxTimecodeProvider = static_cast<URivermaxTimecodeProvider*>(Provider);
+				CurrentTimespan -= FTimespan(0, 0, RmaxTimecodeProvider->UTCSecondsOffset);
+			}
+			else
+			{
+				UE_CALL_ONCE([&] { UE_LOG(LogRivermaxMedia, Warning, TEXT("Rivermax Timecode provider is required for accurate playback.")); });
+			}
+
+			// Rollover 24 hours.
+			CurrentTimespan = FTimespan(CurrentTimespan.GetTicks() % NumberOfTicksPerDay);
+		}
+
+		return CurrentTimespan;
+	}
 	/* FRivermaxVideoPlayer structors
 	 *****************************************************************************/
 
@@ -241,9 +274,13 @@ namespace UE::RivermaxMedia
 
 			UE_LOG(LogRivermaxMedia, Verbose, TEXT("Starting to receive frame '%u' with timestamp %u"), FrameInfo.FrameNumber, FrameInfo.Timestamp);
 
+			// Until PTP Timecode is available sample records frame reception start time for sample picking.
+			Sample->FrameReceptionStart = GetCurrentPTPTimeOfDay();
+
+			// the following will be restored once we have true PTP timecode support.
+			//const double MediaFrameTimeSecs = UE::RivermaxCore::ConvertRTPTimeStampToSeconds(FrameInfo.Timestamp);
+			//Sample->SetTime(FTimespan::FromSeconds(MediaFrameTimeSecs));
 			Sample->SetFrameNumber(FrameInfo.FrameNumber);
-			const double MediaFrameTimeSecs = UE::RivermaxCore::ConvertRTPTimeStampToSecondsTruncated(FrameInfo.Timestamp);
-			Sample->SetTime(FTimespan::FromSeconds(MediaFrameTimeSecs));
 			Sample->SetReceptionState(IRivermaxVideoSample::ESampleState::ReadyForReception);
 			LastFrameToAttemptReception = FrameInfo.FrameNumber;
 			return Sample;
@@ -272,6 +309,10 @@ namespace UE::RivermaxMedia
 		{
 			Sample->GetSampleReceivedEvent()->Trigger();
 		}
+
+		// Until PTP Timecode is available sample records frame reception start time for sample picking.
+		Sample->FrameReceptionEnd = GetCurrentPTPTimeOfDay();
+		
 	}
 
 	void FRivermaxMediaPlayer::TickInput(FTimespan DeltaTime, FTimespan Timecode)
@@ -415,7 +456,7 @@ namespace UE::RivermaxMedia
 		VideoTextureSamplePool = MakeUnique<FRivermaxMediaTextureSamplePool>();
 		Samples->FlushSamples();
 
-		// We only need to store current and the number we delay by.
+		// Only need to store latest frame and the number of frames delayed by.
 		Samples->SetSampleBufferSize(FrameDelay + 1);
 
 
@@ -718,6 +759,77 @@ namespace UE::RivermaxMedia
 		return Sample;
 	}
 
+	TSharedPtr<FMediaIOCoreTextureSampleBase> FRivermaxMediaPlayer::PickSampleToRenderForTimeSynchronized_RenderThread(const FFrameInfo& InFrameInformation)
+	{
+		// Reference time based on evaluation type
+		FTimespan TargetSampleTimespan;
+
+		// Get base uncorrected reference point
+
+		if (EvaluationType == EMediaIOSampleEvaluationType::Timecode && InFrameInformation.RequestedTimecode != FTimecode())
+		{
+			// We'll use timecode data to find a proper sample
+			TargetSampleTimespan = InFrameInformation.RequestedTimecode.ToTimespan(VideoFrameRate);
+		}
+		else
+		{
+			// We'll use platform time to find a proper sample
+			TargetSampleTimespan = InFrameInformation.SampleTimespan;
+		}
+
+		// Apply time correction to the target time
+		const FTimespan RequestedOffsetTimespan = FTimespan::FromSeconds(InFrameInformation.EvaluationOffset);
+
+		// Latency adjusted VSync
+		const FTimespan TargetTimespanCorrected = TargetSampleTimespan - RequestedOffsetTimespan;
+
+		// Go over the sample pool and find a sample closest to the target time
+		int32 ClosestIndex = -1;
+		int64 SmallestInterval(TNumericLimits<int64>::Max());
+
+		// Get all available video samples
+		const TArray<TSharedPtr<IMediaTextureSample>> TextureSamples = Samples->GetVideoSamples();
+		if (TextureSamples.Num() == 0)
+		{
+			return nullptr;
+		}
+
+		for (int32 Index = 0; Index < TextureSamples.Num(); ++Index)
+		{
+			TSharedPtr<FRivermaxMediaTextureSample> Sample = StaticCastSharedPtr<FRivermaxMediaTextureSample>(TextureSamples[Index]);
+
+			// Either closest positive or closest negative
+			const int64 TestInterval = FMath::Abs((Sample->FrameReceptionStart - TargetTimespanCorrected).GetTicks());
+
+			// if VSync is within the time frame of sample's start and end of the reception it is the sample that is returned
+			// otherwise find the sample which began the reception closest to the required VSync.
+			if (TargetTimespanCorrected >= Sample->FrameReceptionStart && TargetTimespanCorrected < Sample->FrameReceptionEnd)
+			{
+				ClosestIndex = Index;
+				break;
+			}
+
+			// '<=' instead of '<' is used here intentionally. Turns out we might have
+			// some samples with the same timecode. To avoid early termination of the search '<=' is used.
+			if (TestInterval <= SmallestInterval)
+			{
+				ClosestIndex = Index;
+				SmallestInterval = TestInterval;
+			}
+			else
+			{
+				// Since our samples are stored in chronological order, it makes no sense
+				// to continue searching. The interval will continue increasing.
+				break;
+			}
+		}
+
+		checkSlow(ClosestIndex >= 0 && ClosestIndex < TextureSamples.Num());
+
+		// Finally, return the closest sample we found
+		return StaticCastSharedPtr<FMediaIOCoreTextureSampleBase, IMediaTextureSample, ESPMode::ThreadSafe>(TextureSamples[ClosestIndex]);
+	}
+
 	void FRivermaxMediaPlayer::PostSampleUsage(FRDGBuilder& GraphBuilder, TSharedPtr<FRivermaxMediaTextureSample> Sample)
 	{
 		GraphBuilder.AddPass(RDG_EVENT_NAME("RivermaxPostSampleUsage"),
diff --git a/Engine/Plugins/VirtualProduction/Rivermax/RivermaxMedia/Source/RivermaxMedia/Private/RivermaxMediaPlayer.h b/Engine/Plugins/VirtualProduction/Rivermax/RivermaxMedia/Source/RivermaxMedia/Private/RivermaxMediaPlayer.h
--- a/Engine/Plugins/VirtualProduction/Rivermax/RivermaxMedia/Source/RivermaxMedia/Private/RivermaxMediaPlayer.h
+++ b/Engine/Plugins/VirtualProduction/Rivermax/RivermaxMedia/Source/RivermaxMedia/Private/RivermaxMediaPlayer.h
@@ -50,7 +50,9 @@ namespace  UE::RivermaxMedia
 		virtual bool JustInTimeSampleRender_RenderThread(FRHICommandListImmediate& RHICmdList, FTextureRHIRef& InDestinationTexture, TSharedPtr<FMediaIOCoreTextureSampleBase>& JITRProxySample) override;
 	protected:
 		virtual TSharedPtr<FMediaIOCoreTextureSampleBase> PickSampleToRenderFramelocked_RenderThread(const FFrameInfo& InFrameInformation) override;
-
+		
+		/**  Pick sample based on Vsync timecode and Start and End of sample reception. */
+		virtual TSharedPtr<FMediaIOCoreTextureSampleBase> PickSampleToRenderForTimeSynchronized_RenderThread(const FFrameInfo& InFrameInformation) override;
 		//~ End FMediaIOCorePlayerBase interface
 
 	public: 
@@ -120,8 +122,8 @@ namespace  UE::RivermaxMedia
 		TRefCountPtr<FRHITexture> CreateIntermediateRenderTarget(FRHICommandListImmediate& RHICmdList, const FIntPoint& InDim, EPixelFormat InPixelFormat, bool bInSRGB);
 
 	private:
-		/** Size of the sample pool. */
-		static constexpr uint32 kMaxNumVideoFrameBuffer = 5;
+		/** Size of the sample pool. The max number of FrameDelay (4) + 2 frames (to give time to return back to the pool) */
+		static constexpr uint32 kMaxNumVideoFrameBuffer = 6;
 
 		/** Current state of the media player. */
 		EMediaState RivermaxThreadNewState;
diff --git a/Engine/Plugins/VirtualProduction/Rivermax/RivermaxMedia/Source/RivermaxMedia/Private/RivermaxMediaTextureSample.h b/Engine/Plugins/VirtualProduction/Rivermax/RivermaxMedia/Source/RivermaxMedia/Private/RivermaxMediaTextureSample.h
--- a/Engine/Plugins/VirtualProduction/Rivermax/RivermaxMedia/Source/RivermaxMedia/Private/RivermaxMediaTextureSample.h
+++ b/Engine/Plugins/VirtualProduction/Rivermax/RivermaxMedia/Source/RivermaxMedia/Private/RivermaxMediaTextureSample.h
@@ -136,6 +136,15 @@ namespace UE::RivermaxMedia
 
 		/** Pooled buffer used for gpudirect functionality. Received content will already be on GPU when received from NIC */
 		TRefCountPtr<FRDGPooledBuffer> GPUBuffer;
+
+	private:
+		friend class FRivermaxMediaPlayer;
+
+		/** The start of the reception marked by the first chunk received by rivermax. */
+		FTimespan FrameReceptionStart = 0;
+
+		/** The end of the reception marked by the last processed packet. */
+		FTimespan FrameReceptionEnd = 0;
 	};
 
 	class FRivermaxMediaTextureSamplePool : public TMediaObjectPool<FRivermaxMediaTextureSample> { };
