Rivermax: fixed a multiple defined symbols complication issue.

#rnx

diff --git a/Engine/Plugins/VirtualProduction/Rivermax/RivermaxMedia/Source/RivermaxMedia/Private/RivermaxMediaUtils.cpp b/Engine/Plugins/VirtualProduction/Rivermax/RivermaxMedia/Source/RivermaxMedia/Private/RivermaxMediaUtils.cpp
--- a/Engine/Plugins/VirtualProduction/Rivermax/RivermaxMedia/Source/RivermaxMedia/Private/RivermaxMediaUtils.cpp
+++ b/Engine/Plugins/VirtualProduction/Rivermax/RivermaxMedia/Source/RivermaxMedia/Private/RivermaxMediaUtils.cpp
@@ -9,7 +9,7 @@
 const FGuid FRivermaxMediaVersion::GUID(0x867597A5, 0x492791B6, 0x3FDACDB0, 0x60344667);
 
 // RivermaxMedia Version registration with core.
-FCustomVersionRegistration GRegisterCustomizableObjectCustomVersion(FRivermaxMediaVersion::GUID, FRivermaxMediaVersion::LatestVersion, TEXT("RivermaxMediaVersion"));
+FCustomVersionRegistration GRegisterRivermaxMediaCustomVersion(FRivermaxMediaVersion::GUID, FRivermaxMediaVersion::LatestVersion, TEXT("RivermaxMediaVersion"));
 
 
 namespace UE::RivermaxMediaUtils::Private
