Rivermax: Fix for an issue where LockedMemory is being nulled when it is about to be used for copying into.

[RN] Rivermax: Fixed an crash caused by inner resizing.

diff --git a/Engine/Plugins/VirtualProduction/Rivermax/RivermaxMedia/Source/RivermaxMedia/Private/RivermaxMediaTextureSample.h b/Engine/Plugins/VirtualProduction/Rivermax/RivermaxMedia/Source/RivermaxMedia/Private/RivermaxMediaTextureSample.h
--- a/Engine/Plugins/VirtualProduction/Rivermax/RivermaxMedia/Source/RivermaxMedia/Private/RivermaxMediaTextureSample.h
+++ b/Engine/Plugins/VirtualProduction/Rivermax/RivermaxMedia/Source/RivermaxMedia/Private/RivermaxMediaTextureSample.h
@@ -55,8 +55,12 @@ namespace UE::RivermaxMedia
 			// This means that the sample is being held and managed by the Player
 			FScopeLock Lock(&StateChangeCriticalSecion);
 
+			while (IsAwaitingForGPUTransfer())
+			{
+				FPlatformProcess::SleepNoStats(0.);
+			}
+
 			// When this sample is returned back to the pool, it means that it is done rendering and is released from sample container.
-			SetAwaitingForGPUTransfer(false);
 			SetReceptionState(IRivermaxSample::ESampleState::ReadyForReception);
 			if (SampleConversionFence.IsValid())
 			{
