Rivermax: Removed an early exit from WaitForSample that depended on Player Play status. Added a check to ShutdownPoolable that makes sure that all rendering processies related to the sample are done.

#rnx

diff --git a/Engine/Plugins/VirtualProduction/Rivermax/RivermaxMedia/Source/RivermaxMedia/Private/RivermaxMediaPlayer.cpp b/Engine/Plugins/VirtualProduction/Rivermax/RivermaxMedia/Source/RivermaxMedia/Private/RivermaxMediaPlayer.cpp
--- a/Engine/Plugins/VirtualProduction/Rivermax/RivermaxMedia/Source/RivermaxMedia/Private/RivermaxMediaPlayer.cpp
+++ b/Engine/Plugins/VirtualProduction/Rivermax/RivermaxMedia/Source/RivermaxMedia/Private/RivermaxMediaPlayer.cpp
@@ -947,11 +947,6 @@ namespace UE::RivermaxMedia
 
 		while (true)
 		{
-			if (!IsReadyToPlay())
-			{
-				return false;
-			}
-
 			// Our goal here is to wait until the expected frame is available to be used (received) unless there is a timeout
 			if (WaitConditionFunction(Sample))
 			{
diff --git a/Engine/Plugins/VirtualProduction/Rivermax/RivermaxMedia/Source/RivermaxMedia/Private/RivermaxMediaTextureSample.h b/Engine/Plugins/VirtualProduction/Rivermax/RivermaxMedia/Source/RivermaxMedia/Private/RivermaxMediaTextureSample.h
--- a/Engine/Plugins/VirtualProduction/Rivermax/RivermaxMedia/Source/RivermaxMedia/Private/RivermaxMediaTextureSample.h
+++ b/Engine/Plugins/VirtualProduction/Rivermax/RivermaxMedia/Source/RivermaxMedia/Private/RivermaxMediaTextureSample.h
@@ -52,13 +52,12 @@ namespace UE::RivermaxMedia
 		virtual void InitializePoolable() override {};
 		virtual void ShutdownPoolable() override 
 		{
+			// When ShutdownPoolable is called, it means that that the sample is returned back to the pool or it is done being rendered.
+			// We should never be in the state where this awaits GPU transfer while being released.
+			checkf(IsAwaitingForGPUTransfer() == false, TEXT("This Sample is still transferring data to GPU. If this hits, something went wrong."));
 			// This means that the sample is being held and managed by the Player
 			FScopeLock Lock(&StateChangeCriticalSecion);
 
-			while (IsAwaitingForGPUTransfer())
-			{
-				FPlatformProcess::SleepNoStats(0.);
-			}
 
 			// When this sample is returned back to the pool, it means that it is done rendering and is released from sample container.
 			SetReceptionState(IRivermaxSample::ESampleState::ReadyForReception);
@@ -118,7 +117,7 @@ namespace UE::RivermaxMedia
 
 	public:
 		/** Locked memory of gpu buffer when uploading */
-		void* LockedMemory = nullptr;
+		std::atomic<void*> LockedMemory = nullptr;
 
 		/** Write fence enqueued after sample conversion to know when it's ready to be reused */
 		FGPUFenceRHIRef SampleConversionFence;
