Rivermax: Stability:
- GPU direct on input should now be fully supported. Better tracking of the sample states and better handling of waits on input.
- Fixed an issue where incorrect Frame number was inferred from Timestamp. As an example lets take timestamp 184500 and framerate 60fps The resulting frame number inferred by the function 122.999999999, which was floored to 122. This was causing ocassional invisible "skipped" frames in rare cases.

[RN] Rivermax: Fixed bugs with tearing with GPU Direct on input enabled. Fixed issues with acquisition of invalid framenumbers due to floating point precision issues.

diff --git a/Engine/Plugins/VirtualProduction/Rivermax/RivermaxCore/Source/RivermaxCore/Private/RivermaxUtils.cpp b/Engine/Plugins/VirtualProduction/Rivermax/RivermaxCore/Source/RivermaxCore/Private/RivermaxUtils.cpp
--- a/Engine/Plugins/VirtualProduction/Rivermax/RivermaxCore/Source/RivermaxCore/Private/RivermaxUtils.cpp
+++ b/Engine/Plugins/VirtualProduction/Rivermax/RivermaxCore/Source/RivermaxCore/Private/RivermaxUtils.cpp
@@ -13,7 +13,7 @@ namespace UE::RivermaxCore::Private::Utils
 	{
 		using namespace UE::RivermaxCore::Private::Utils;
 		const double MediaFrameTime = Timestamp / MediaClockSampleRate;
-		const uint32 FrameNumber = FMath::Floor(MediaFrameTime * FrameRate.AsDecimal());
+		const uint32 FrameNumber = FMath::RoundToInt(MediaFrameTime * FrameRate.AsDecimal());
 		return FrameNumber;
 	}
 
diff --git a/Engine/Plugins/VirtualProduction/Rivermax/RivermaxCore/Source/RivermaxCore/Public/IRivermaxInputStream.h b/Engine/Plugins/VirtualProduction/Rivermax/RivermaxCore/Source/RivermaxCore/Public/IRivermaxInputStream.h
--- a/Engine/Plugins/VirtualProduction/Rivermax/RivermaxCore/Source/RivermaxCore/Public/IRivermaxInputStream.h
+++ b/Engine/Plugins/VirtualProduction/Rivermax/RivermaxCore/Source/RivermaxCore/Public/IRivermaxInputStream.h
@@ -80,11 +80,21 @@ namespace UE::RivermaxCore
 		}
 		enum class ESampleState : uint8
 		{
+			// Sample hasn't started receiving and is awaiting to be used. In this state it contains the data previously written into it.
+			Idle,
+
 			// Sample is ready to be written into.
 			ReadyForReception,
 
-			// Error receiving sample.
-			ReceptionError
+			// Sample has received the data. And ready to be used for rendering.
+			Received,
+
+			// Error receiving sample. Used at the same point as Received state.
+			ReceptionError,
+
+			// After sample is received and is ready to be rendered and until the GPU is done with the sample this is the state the sample is in. 
+			// Currently only used for debugging purposes.
+			Rendering
 		};
 
 		/**
diff --git a/Engine/Plugins/VirtualProduction/Rivermax/RivermaxMedia/Source/RivermaxMedia/Private/RivermaxMediaPlayer.cpp b/Engine/Plugins/VirtualProduction/Rivermax/RivermaxMedia/Source/RivermaxMedia/Private/RivermaxMediaPlayer.cpp
--- a/Engine/Plugins/VirtualProduction/Rivermax/RivermaxMedia/Source/RivermaxMedia/Private/RivermaxMediaPlayer.cpp
+++ b/Engine/Plugins/VirtualProduction/Rivermax/RivermaxMedia/Source/RivermaxMedia/Private/RivermaxMediaPlayer.cpp
@@ -252,7 +252,7 @@ namespace UE::RivermaxMedia
 							PreviousSample->SetFrameNumber(PreviousFrameNumber);
 
 							// this will stop the waiting thread from copying invalid data and waiting for the copy.
-							PreviousSample->SetReceptionState(IRivermaxVideoSample::ESampleState::ReceptionError);
+							PreviousSample->SetReceptionState(IRivermaxSample::ESampleState::ReceptionError);
 
 							// this will stop the waiting thread from waiting for the sample reception.
 							PreviousSample->GetSampleReceivedEvent()->Trigger();
@@ -261,14 +261,20 @@ namespace UE::RivermaxMedia
 				}
 
 				Sample = FrameLockedSamples[FRAME_NUM_TO_INDEX(FrameInfo.FrameNumber)];
-				Sample->GetSampleReceivedEvent()->Reset();
-
 			}
 			else
 			{
 				Sample = VideoTextureSamplePool->AcquireShared(false /*NoAllocation*/);
 			}
 
+			// The sample hasn't completed its cycle. This is a sanity check to ensure nothing unexpected has happened with the sample.
+			// If this happens it is fine to continue since the sample most likely was skipped by the rendering step.
+			if (Sample->GetReceptionState() != IRivermaxVideoSample::ESampleState::Idle)
+			{
+				const TCHAR*& SampleStateString = FRivermaxMediaTextureSample::SampleStateToStringRef(Sample->GetReceptionState());
+				UE_LOG(LogRivermaxMedia, Warning, TEXT("The sample hasn't completed it's cycle. The frame number of the incomplete Sample: %u Frame Number of the sample about to be received: %u, State: %s"), Sample->GetFrameNumber(), FrameInfo.FrameNumber, SampleStateString);
+			}
+
 			// (More of a sanity check. Shouldn't be in this state.).
 			if (!Sample.IsValid())
 			{
@@ -276,6 +282,13 @@ namespace UE::RivermaxMedia
 				return nullptr;
 			}
 
+			// With Unreal as the sender the receiver shouldn't be in situations where it receives the same number twice in a row. However it might be different with other devices.
+			// in case such situation is encountered it is good to have something logged.
+			if (LastFrameToAttemptReception == FrameInfo.FrameNumber) 
+			{
+				UE_LOG(LogRivermaxMedia, Warning, TEXT("The same frame number has been received twice in a row. Frame Number: %u, Timestamp: %u"), FrameInfo.FrameNumber, FrameInfo.Timestamp);
+			}
+
 			UE_LOG(LogRivermaxMedia, Verbose, TEXT("Starting to receive frame '%u' with timestamp %u"), FrameInfo.FrameNumber, FrameInfo.Timestamp);
 
 			// Until PTP Timecode is available sample records frame reception start time for sample picking.
@@ -284,9 +297,10 @@ namespace UE::RivermaxMedia
 			// the following will be restored once we have true PTP timecode support.
 			//const double MediaFrameTimeSecs = UE::RivermaxCore::ConvertRTPTimeStampToSeconds(FrameInfo.Timestamp);
 			//Sample->SetTime(FTimespan::FromSeconds(MediaFrameTimeSecs));
+			Sample->SetReceptionState(IRivermaxSample::ESampleState::ReadyForReception);
 			Sample->SetFrameNumber(FrameInfo.FrameNumber);
-			Sample->SetReceptionState(IRivermaxVideoSample::ESampleState::ReadyForReception);
 			LastFrameToAttemptReception = FrameInfo.FrameNumber;
+
 			return Sample;
 		}
 
@@ -314,6 +328,8 @@ namespace UE::RivermaxMedia
 			Sample->GetSampleReceivedEvent()->Trigger();
 		}
 
+		Sample->SetReceptionState(IRivermaxSample::ESampleState::Received);
+
 		// Until PTP Timecode is available sample records frame reception start time for sample picking.
 		Sample->FrameReceptionEnd = GetCurrentPTPTimeOfDay();
 		
@@ -471,7 +487,8 @@ namespace UE::RivermaxMedia
 			NewSample = VideoTextureSamplePool->AcquireShared(true /*Allocate New*/);
 			NewSample->InitializeGPUBuffer(InResolution, DesiredPixelFormat, bStreamSupportsGPUDirect);
 			NewSample->SampleConversionFence = RHICreateGPUFence(*FString::Printf(TEXT("RmaxConversionDoneFence_%02d"), Index));
-			NewSample->SetReceptionState(IRivermaxSample::ESampleState::ReadyForReception);
+
+			NewSample->SetReceptionState(IRivermaxSample::ESampleState::Idle);
 			NewSample->SetInputFormat(DesiredPixelFormat);
 			if (IS_FRAMELOCKED())
 			{
@@ -889,12 +906,12 @@ namespace UE::RivermaxMedia
 		UE_LOG(LogRivermaxMedia, Warning, TEXT("Issue receiving frame number %u."), Sample->GetFrameNumber());
 		if (Sample->IsBeingRendered())
 		{
-			Sample->SetReceptionState(FRivermaxMediaTextureSample::ESampleState::ReceptionError);
+			Sample->SetReceptionState(IRivermaxSample::ESampleState::ReceptionError);
 			Sample->GetSampleReceivedEvent()->Trigger();
 		}
 		else
 		{
-			Sample->SetReceptionState(FRivermaxMediaTextureSample::ESampleState::ReadyForReception);
+			Sample->SetReceptionState(IRivermaxSample::ESampleState::Idle);
 		}
 	}
 
@@ -944,10 +961,9 @@ namespace UE::RivermaxMedia
 		}
 	}
 
-	bool FRivermaxMediaPlayer::WaitForSample(TSharedPtr<FRivermaxMediaTextureSample> Sample, FWaitConditionFunc WaitConditionFunction, bool bCanTimeout)
+	bool FRivermaxMediaPlayer::WaitForSample(const TSharedPtr<FRivermaxMediaTextureSample>& Sample, const uint64 AwaitingFrameNumber, FWaitConditionFunc WaitConditionFunction, const double TimeoutSeconds)
 	{
 		const double StartTimeSeconds = FPlatformTime::Seconds();
-		const double TimeoutSeconds = CVarRivermaxDefaultThreadBlockTimeoutSec.GetValueOnAnyThread();
 
 		while (true)
 		{
@@ -957,11 +973,12 @@ namespace UE::RivermaxMedia
 				return true;
 			}
 
+			FPlatformProcess::SleepNoStats(SleepTimeSeconds);
+
 			{
-				FPlatformProcess::SleepNoStats(SleepTimeSeconds);
-				if (bCanTimeout && ((FPlatformTime::Seconds() - StartTimeSeconds) > TimeoutSeconds))
+				if (((FPlatformTime::Seconds() - StartTimeSeconds) > TimeoutSeconds))
 				{
-					UE_LOG(LogRivermaxMedia, Error, TEXT("Timed out waiting for frame %u."), Sample->GetFrameNumber());
+					UE_LOG(LogRivermaxMedia, Error, TEXT("Timed out waiting for frame %u."), AwaitingFrameNumber);
 					return false;
 				}
 			}
@@ -978,7 +995,10 @@ namespace UE::RivermaxMedia
 
 		OutConverterSetup.PreConvertFunc = [SamplePtr = Sample, NextFrameExpectations, this](const FRDGBuilder& GraphBuilder)
 			{
-				TRACE_CPUPROFILER_EVENT_SCOPE_TEXT(*FString::Printf(TEXT("Rmax::01_PreConvertFunction %u"), (SamplePtr->GetFrameNumber())));
+				TRACE_CPUPROFILER_EVENT_SCOPE_TEXT(*FString::Printf(TEXT("Rmax::01_PreConvertFunction %u"), NextFrameExpectations));
+
+				SamplePtr->SetAwaitingForGPUTransfer(true);
+
 				// When GPUDirect is not involved, we have an extra step to do. We need to wait for the sample to be received
 				// but also initiate the memcopy to gpu memory for it to be rendered
 				if (bStreamSupportsGPUDirect == false)
@@ -986,7 +1006,6 @@ namespace UE::RivermaxMedia
 					constexpr uint32 Offset = 0;
 					const uint32 Size = SamplePtr->GetGPUBuffer()->GetSize();
 
-					SamplePtr->SetAwaitingForGPUTransfer(true);
 
 					// Todo: Find a way to keep the buffer permanently locked.
 					if (SamplePtr->LockedMemory == nullptr)
@@ -1002,7 +1021,8 @@ namespace UE::RivermaxMedia
 							{
 								--TasksInFlight;
 							};
-							TRACE_CPUPROFILER_EVENT_SCOPE_TEXT(*FString::Printf(TEXT("Rmax::01_CopyPixels %u"), (NextFrameExpectations)));
+							TRACE_CPUPROFILER_EVENT_SCOPE(RmaxWaitAndCopyPixels);
+							TRACE_CPUPROFILER_EVENT_SCOPE_TEXT(*FString::Printf(TEXT("Rmax::WaitAndCopyPixels %u"), (NextFrameExpectations)));
 
 							// Indicates if we sample reception started on the input thread.
 							bool bStreamReceptionStarted = true;
@@ -1011,34 +1031,28 @@ namespace UE::RivermaxMedia
 							// This waits for frame locked samples to start writing packets and then for the completion of the reception.
 							if (IS_FRAMELOCKED())
 							{
-								TRACE_CPUPROFILER_EVENT_SCOPE_TEXT(*FString::Printf(TEXT("Rmax::01_WaitForReception %u"), (NextFrameExpectations)));
+								TRACE_CPUPROFILER_EVENT_SCOPE(RmaxWaitForPixels);
 
-								// Waiting to start receiving.
+								// Waiting for reception.
 								const double TimeoutSeconds = CVarRivermaxDefaultThreadBlockTimeoutSec.GetValueOnAnyThread();
 								const double StartTimeSeconds = FPlatformTime::Seconds();
-								uint64 CurrentSampleFrameNumber = SamplePtr->GetFrameNumber();
 
-								while (CurrentSampleFrameNumber < NextFrameExpectations)
+								const double OneFrameTime = 1. / GetFrameRate().AsDecimal();
+
+								// This loop will continue until it times out or the frame number matches or exceeds the expected frame number and the sample has been received or errored out.
+								while (SamplePtr->GetFrameNumber() < NextFrameExpectations ||
+									(SamplePtr->GetReceptionState() != IRivermaxSample::ESampleState::Received && SamplePtr->GetReceptionState() != IRivermaxSample::ESampleState::ReceptionError))
 								{
-									FPlatformProcess::SleepNoStats(SleepTimeSeconds);
 									if ((FPlatformTime::Seconds() - StartTimeSeconds) > TimeoutSeconds)
 									{
 										UE_LOG(LogRivermaxMedia, Warning, TEXT("Timed out waiting for frame #%u to start being received."), NextFrameExpectations);
 										bStreamReceptionStarted = false;
 										break;
 									}
-									CurrentSampleFrameNumber = SamplePtr->GetFrameNumber();
-								}
 
-								// At this point it is known that the sample reception has started or not.
-								if (bStreamReceptionStarted)
-								{
-									TRACE_CPUPROFILER_EVENT_SCOPE(RmaxWaitSampleReception);
-									TRACE_CPUPROFILER_EVENT_SCOPE_TEXT(*FString::Printf(TEXT("Rmax::02_WaitSampleReceived %u"), (NextFrameExpectations)));
-									// We wait for 2 Media Player frames to timeout. Reason: Due to 2110 standard anything more than 1 frame is going to be discarded.
-									// Since we are guaranteed to have started receiving the buffer we shouldn't wait longer than that.
-									const double TwoFrameTime = 2. / GetFrameRate().AsDecimal();
-									SamplePtr->GetSampleReceivedEvent()->Wait(FTimespan::FromSeconds(TwoFrameTime));
+									// This could turn into active loop only if frame number doesn't match and the sample is signaled.
+									// However such a case can only happen if the sample was received but wasn't rendered previously.
+									SamplePtr->GetSampleReceivedEvent()->Wait(FTimespan::FromSeconds(OneFrameTime));
 								}
 							}
 
@@ -1050,7 +1064,7 @@ namespace UE::RivermaxMedia
 							// In case reception failed mid way.
 							if (bStreamReceptionStarted && SamplePtr->GetReceptionState() != IRivermaxSample::ESampleState::ReceptionError)
 							{
-								TRACE_CPUPROFILER_EVENT_SCOPE_TEXT(*FString::Printf(TEXT("Rmax::03_CopyBufferToGPU %u"), (SamplePtr->GetFrameNumber())));
+								TRACE_CPUPROFILER_EVENT_SCOPE(CopyBufferToGPU);
 								FMemory::Memcpy(SamplePtr->LockedMemory, SamplePtr->GetOrRequestBuffer(Size), Size);
 							}
 							else
@@ -1058,6 +1072,74 @@ namespace UE::RivermaxMedia
 								UE_LOG(LogRivermaxMedia, Warning, TEXT("Incomplete or failed pixels will be rendered for frame %u"), NextFrameExpectations);
 							}
 
+							// Set this for debugging purposes to know exactly when this sample is shifted to rendering state.
+							SamplePtr->SetReceptionState(IRivermaxSample::ESampleState::Rendering);
+
+							// Signals the RHI thread that the GPU transfer has completed.
+							SamplePtr->SetAwaitingForGPUTransfer(false);
+						}
+					);
+				}
+
+				// Stream supports GPU Direct
+				else
+				{
+					++TasksInFlight;
+					UE::Tasks::Launch(UE_SOURCE_LOCATION,
+						[SamplePtr, NextFrameExpectations, this]()
+						{
+							ON_SCOPE_EXIT
+							{
+								--TasksInFlight;
+							};
+
+							TRACE_CPUPROFILER_EVENT_SCOPE(RmaxWaitSampleReceptionGPUDirect);
+							TRACE_CPUPROFILER_EVENT_SCOPE_TEXT(*FString::Printf(TEXT("Rmax::02_WaitSampleReceivedGpuDirect %u"), (NextFrameExpectations)));
+
+							double StartTime = FPlatformTime::Seconds();
+							// In frame locked mode the packets potentially haven't started being received yet. 
+							// This waits for frame locked samples to start writing packets and then for the completion of the reception.
+							if (IS_FRAMELOCKED())
+							{
+								// Waiting for reception.
+								const double TimeoutSeconds = CVarRivermaxDefaultThreadBlockTimeoutSec.GetValueOnAnyThread();
+								const double StartTimeSeconds = FPlatformTime::Seconds();
+
+								const double OneFrameTime = 1. / GetFrameRate().AsDecimal();
+
+								// This loop will continue until it times out or the frame number matches or exceeds the expected frame number and the sample has been received or errored out.
+								while (SamplePtr->GetFrameNumber() < NextFrameExpectations || 
+									(SamplePtr->GetReceptionState() != IRivermaxSample::ESampleState::Received && SamplePtr->GetReceptionState() != IRivermaxSample::ESampleState::ReceptionError))
+								{
+									if ((FPlatformTime::Seconds() - StartTimeSeconds) > TimeoutSeconds)
+									{
+										UE_LOG(LogRivermaxMedia, Warning, TEXT("Timed out waiting for frame #%u to start being received."), NextFrameExpectations);
+										break;
+									}
+
+									// This could turn into active loop if frame number doesn't match, but the sample is signaled.
+									// However such a case can only happen if the sample was received but wasn't rendered previously.
+									SamplePtr->GetSampleReceivedEvent()->Wait(FTimespan::FromSeconds(OneFrameTime));
+								}
+							}
+
+							if (SamplePtr->GetFrameNumber() != NextFrameExpectations)
+							{
+								double Elapsed = FPlatformTime::Seconds() - StartTime;
+								const TCHAR*& SampleStateString = FRivermaxMediaTextureSample::SampleStateToStringRef(SamplePtr->GetReceptionState());
+
+								UE_LOG(LogRivermaxMedia, Warning, TEXT("1. Rendering unexpected frame %u, when frame %u was expected. Elapsed wait time %f. State: %s"), SamplePtr->GetFrameNumber(), NextFrameExpectations, Elapsed, SampleStateString);
+							}
+
+							// In case reception failed mid way.
+							if (SamplePtr->GetReceptionState() == IRivermaxSample::ESampleState::ReceptionError)
+							{
+								UE_LOG(LogRivermaxMedia, Warning, TEXT("Incomplete or failed pixels will be rendered for frame %u"), NextFrameExpectations);
+							}
+
+							// Set this for debugging purposes to know exactly when this sample is shifted to rendering state.
+							SamplePtr->SetReceptionState(IRivermaxSample::ESampleState::Rendering);
+
 							// Signals the RHI thread that the GPU transfer has completed.
 							SamplePtr->SetAwaitingForGPUTransfer(false);
 						}
@@ -1065,6 +1147,7 @@ namespace UE::RivermaxMedia
 				}
 
 
+
 				RHI_BREADCRUMB_EVENT_STAT(GraphBuilder.RHICmdList, Rmax_WaitForPixels, "Rmax::WaitForPixels");
 				SCOPED_GPU_STAT(GraphBuilder.RHICmdList, Rmax_WaitForPixels);
 				// Since we are going to enqueue a lambda that can potentially sleep in the RHI thread if the pixels haven't arrived,
@@ -1082,10 +1165,15 @@ namespace UE::RivermaxMedia
 
 						TRACE_CPUPROFILER_EVENT_SCOPE_TEXT(*FString::Printf(TEXT("Rmax::04_WaitForUploadToFinish %u"), (NextFrameExpectations)));
 
-						// Only accept timeouts here if we're going through the system memory in order to avoid timing out while doing
-						// the memcopy (buffer upload)
-						bool bCanTimeout = bStreamSupportsGPUDirect;
-						WaitForSample(SamplePtr, MoveTemp(ReadyToRenderConditionFunc), bCanTimeout);
+						// Set timeout to 10 seconds which should be enough to make sure that the frame is received or errors out.
+						const double TimeoutSeconds = 10.;
+						WaitForSample(SamplePtr, NextFrameExpectations, MoveTemp(ReadyToRenderConditionFunc), TimeoutSeconds);
+
+						if (SamplePtr->GetFrameNumber() != NextFrameExpectations)
+						{
+							UE_LOG(LogRivermaxMedia, Warning, TEXT("2. Rendering unexpected frame %u, when frame %u was expected."), SamplePtr->GetFrameNumber(), NextFrameExpectations);
+						}
+
 					}
 				);
 
diff --git a/Engine/Plugins/VirtualProduction/Rivermax/RivermaxMedia/Source/RivermaxMedia/Private/RivermaxMediaPlayer.h b/Engine/Plugins/VirtualProduction/Rivermax/RivermaxMedia/Source/RivermaxMedia/Private/RivermaxMediaPlayer.h
--- a/Engine/Plugins/VirtualProduction/Rivermax/RivermaxMedia/Source/RivermaxMedia/Private/RivermaxMediaPlayer.h
+++ b/Engine/Plugins/VirtualProduction/Rivermax/RivermaxMedia/Source/RivermaxMedia/Private/RivermaxMediaPlayer.h
@@ -107,7 +107,7 @@ namespace  UE::RivermaxMedia
 
 		/** Function waiting for the expected frame to be received. */
 		using FWaitConditionFunc = TUniqueFunction<bool(const TSharedPtr<FRivermaxMediaTextureSample>&)>;
-		bool WaitForSample(TSharedPtr<FRivermaxMediaTextureSample> Sample, FWaitConditionFunc WaitConditionFunction, bool bCanTimeout);
+		bool WaitForSample(const TSharedPtr<FRivermaxMediaTextureSample>& Sample, const uint64 AwaitingFrameNumber, FWaitConditionFunc WaitConditionFunction, const double TimeoutSeconds);
 		
 		/** Called after sample was converted / rendered. Used write a fence to detect when sample is reusable */
 		void PostSampleUsage(FRDGBuilder& GraphBuilder, TSharedPtr<FRivermaxMediaTextureSample> Sample);
diff --git a/Engine/Plugins/VirtualProduction/Rivermax/RivermaxMedia/Source/RivermaxMedia/Private/RivermaxMediaTextureSample.cpp b/Engine/Plugins/VirtualProduction/Rivermax/RivermaxMedia/Source/RivermaxMedia/Private/RivermaxMediaTextureSample.cpp
--- a/Engine/Plugins/VirtualProduction/Rivermax/RivermaxMedia/Source/RivermaxMedia/Private/RivermaxMediaTextureSample.cpp
+++ b/Engine/Plugins/VirtualProduction/Rivermax/RivermaxMedia/Source/RivermaxMedia/Private/RivermaxMediaTextureSample.cpp
@@ -10,18 +10,16 @@
 
 namespace UE::RivermaxMedia
 {
-
 FRivermaxMediaTextureSample::FRivermaxMediaTextureSample()
 	: FMediaIOCoreTextureSampleBase()
+	, SampleReceivedEvent(FPlatformProcess::GetSynchEventFromPool(true))
 {
-	SampleReceivedEvent = FGenericPlatformProcess::GetSynchEventFromPool(false);
 }
 
 FRivermaxMediaTextureSample::~FRivermaxMediaTextureSample()
 {
 	FGenericPlatformProcess::ReturnSynchEventToPool(SampleReceivedEvent);
 }
-
 const FMatrix& FRivermaxMediaTextureSample::GetYUVToRGBMatrix() const
 {
 	return MediaShaders::YuvToRgbRec709Scaled;
@@ -83,6 +81,18 @@ uint8* FRivermaxMediaTextureSample::GetVideoBufferRawPtr(uint32 VideoBufferSize)
 	return reinterpret_cast<uint8*>(GetOrRequestBuffer(VideoBufferSize));
 }
 
+void FRivermaxMediaTextureSample::SetReceptionState(ESampleState NewState)
+{
+	if (UE_LOG_ACTIVE(LogRivermaxMedia, VeryVerbose))
+	{
+		const TCHAR*& PreviousStateString = SampleStateToStringRef(GetReceptionState());
+		const TCHAR*& NewStateString = SampleStateToStringRef(NewState);
+		UE_LOG(LogRivermaxMedia, VeryVerbose, TEXT("Changing state for frame number: %u , Preivous State: %s, New state: %s"), GetFrameNumber(), PreviousStateString, NewStateString);
+	}
+	
+	UE::RivermaxCore::IRivermaxVideoSample::SetReceptionState(NewState);
+}
+
 ERivermaxMediaSourcePixelFormat FRivermaxMediaTextureSample::GetInputFormat() const
 {
 	return InputFormat;
diff --git a/Engine/Plugins/VirtualProduction/Rivermax/RivermaxMedia/Source/RivermaxMedia/Private/RivermaxMediaTextureSample.h b/Engine/Plugins/VirtualProduction/Rivermax/RivermaxMedia/Source/RivermaxMedia/Private/RivermaxMediaTextureSample.h
--- a/Engine/Plugins/VirtualProduction/Rivermax/RivermaxMedia/Source/RivermaxMedia/Private/RivermaxMediaTextureSample.h
+++ b/Engine/Plugins/VirtualProduction/Rivermax/RivermaxMedia/Source/RivermaxMedia/Private/RivermaxMediaTextureSample.h
@@ -8,6 +8,7 @@
 #include "MediaIOCoreTextureSampleBase.h"
 #include "MediaIOCoreSamples.h"
 #include "MediaShaders.h"
+#include "RivermaxMediaLog.h"
 #include "RivermaxMediaSource.h"
 #include "Templates/SharedPointer.h"
 
@@ -24,6 +25,16 @@ namespace UE::RivermaxMedia
 		FRivermaxMediaTextureSample();
 		virtual ~FRivermaxMediaTextureSample() override;
 
+		/** A helper function that returns the string based on the provided state. */
+		static inline const TCHAR*& SampleStateToStringRef(ESampleState InState)
+		{
+			static const TCHAR* States[] = {
+				TEXT("Idle"), TEXT("ReadyForReception"), TEXT("Received"), TEXT("ReceptionError"), TEXT("Rendering")
+			};
+
+			return States[(uint8)InState];
+		}
+
 		//~ Begin IMediaTextureSample interface
 		virtual bool IsCacheable() const override;
 		virtual const FMatrix& GetYUVToRGBMatrix() const override;
@@ -48,8 +59,10 @@ namespace UE::RivermaxMedia
 		// Inherited via IRivermaxVideoSample
 		virtual TRefCountPtr<FRDGPooledBuffer> GetGPUBuffer() const override;
 		virtual uint8* GetVideoBufferRawPtr(uint32 VideoBufferSize) override;
+		virtual void SetReceptionState(ESampleState State) override;
 
 		virtual void InitializePoolable() override {};
+
 		virtual void ShutdownPoolable() override 
 		{
 			// When ShutdownPoolable is called, it means that that the sample is returned back to the pool or it is done being rendered.
@@ -58,9 +71,11 @@ namespace UE::RivermaxMedia
 			// This means that the sample is being held and managed by the Player
 			FScopeLock Lock(&StateChangeCriticalSecion);
 
+			// This sample is done and anything that waits for this sample from now on should wait until it is received again.
+			GetSampleReceivedEvent()->Reset();
 
 			// When this sample is returned back to the pool, it means that it is done rendering and is released from sample container.
-			SetReceptionState(IRivermaxSample::ESampleState::ReadyForReception);
+			SetReceptionState(IRivermaxSample::ESampleState::Idle);
 			if (SampleConversionFence.IsValid())
 			{
 				SampleConversionFence->Clear();
@@ -110,7 +125,7 @@ namespace UE::RivermaxMedia
 			return bIsPendingRendering;
 		};
 
-		FEvent* GetSampleReceivedEvent() 
+		FEvent* GetSampleReceivedEvent()
 		{
 			return SampleReceivedEvent;
 		}
