Rivermax: Fixed a GPUDirect issue where the received frame was always behind. The issue was due to the Rivermax Media Capture not waiting for GPU work to be done.

[RN] Rivermax: GPUDirect on input should now be supported.

diff --git a/Engine/Plugins/VirtualProduction/Rivermax/RivermaxCore/Source/RivermaxCore/Private/Streams/RivermaxOutAncStream.cpp b/Engine/Plugins/VirtualProduction/Rivermax/RivermaxCore/Source/RivermaxCore/Private/Streams/RivermaxOutAncStream.cpp
--- a/Engine/Plugins/VirtualProduction/Rivermax/RivermaxCore/Source/RivermaxCore/Private/Streams/RivermaxOutAncStream.cpp
+++ b/Engine/Plugins/VirtualProduction/Rivermax/RivermaxCore/Source/RivermaxCore/Private/Streams/RivermaxOutAncStream.cpp
@@ -78,15 +78,9 @@ namespace UE::RivermaxCore::Private
 		// Clear reserved frame if there is one. If not, try to get free frame.
 		// ReservedFrame should always be valid when block on reservation mode is used.
 		TSharedPtr<FRivermaxOutputFrame> ReservedFrame;
-		if (ReservedFrames.Dequeue(ReservedFrame))
-		{
-			// Sanity check.
-			//check(FrameInfo->FrameIdentifier == ReservedFrame->GetFrameCounter());
-		}
-		else
-		{
-			check(Options.FrameLockingMode != EFrameLockingMode::BlockOnReservation);
-		}
+
+		// At the moment we always have reserved frame for ANC
+		check(ReservedFrames.RemoveAndCopyValue(FrameInfo->FrameIdentifier, ReservedFrame))
 
 		// If this is invalid it means that frame locking mode is BlockOnReservation and the render ran faster than media output fps.
 		FrameInfoToSend->FrameIdentifier = FrameInfo->FrameIdentifier;
@@ -186,7 +180,7 @@ namespace UE::RivermaxCore::Private
 		if (ReservedFrame.IsValid())
 		{
 			ReservedFrame->SetFrameCounter(FrameCounter);
-			ReservedFrames.Enqueue(ReservedFrame);
+			ReservedFrames.Add(FrameCounter, ReservedFrame);
 		}
 
 		return ReservedFrame.IsValid();
diff --git a/Engine/Plugins/VirtualProduction/Rivermax/RivermaxCore/Source/RivermaxCore/Private/Streams/RivermaxOutStream.h b/Engine/Plugins/VirtualProduction/Rivermax/RivermaxCore/Source/RivermaxCore/Private/Streams/RivermaxOutStream.h
--- a/Engine/Plugins/VirtualProduction/Rivermax/RivermaxCore/Source/RivermaxCore/Private/Streams/RivermaxOutStream.h
+++ b/Engine/Plugins/VirtualProduction/Rivermax/RivermaxCore/Source/RivermaxCore/Private/Streams/RivermaxOutStream.h
@@ -473,7 +473,7 @@ namespace UE::RivermaxCore::Private
 		* Frames reserved when Media Capture texture is about to be converted into Buffer to be sent.
 		* Key is frame counter
 		*/
-		mutable TQueue<TSharedPtr<FRivermaxOutputFrame>> ReservedFrames;
+		mutable TMap<uint64, TSharedPtr<FRivermaxOutputFrame>> ReservedFrames;
 
 		/** Thread scheduling frame output */
 		TUniquePtr<FRunnableThread> RivermaxThread;
diff --git a/Engine/Plugins/VirtualProduction/Rivermax/RivermaxCore/Source/RivermaxCore/Private/Streams/RivermaxOutVideoStream.cpp b/Engine/Plugins/VirtualProduction/Rivermax/RivermaxCore/Source/RivermaxCore/Private/Streams/RivermaxOutVideoStream.cpp
--- a/Engine/Plugins/VirtualProduction/Rivermax/RivermaxCore/Source/RivermaxCore/Private/Streams/RivermaxOutVideoStream.cpp
+++ b/Engine/Plugins/VirtualProduction/Rivermax/RivermaxCore/Source/RivermaxCore/Private/Streams/RivermaxOutVideoStream.cpp
@@ -426,12 +426,7 @@ namespace UE::RivermaxCore::Private
 		// Clear reserved frame if there is one. If not, try to get free frame.
 		// ReservedFrame should always be valid when block on reservation mode is used.
 		TSharedPtr<FRivermaxOutputFrame> ReservedFrame;
-		if (ReservedFrames.Dequeue(ReservedFrame))
-		{
-			// Sanity check.
-			//check(FrameInfo->FrameIdentifier == ReservedFrame->GetFrameCounter());
-		}
-		else
+		if (!ReservedFrames.RemoveAndCopyValue(FrameInfo->FrameIdentifier, ReservedFrame))
 		{
 			check(Options.FrameLockingMode != EFrameLockingMode::BlockOnReservation);
 			ReservedFrame = FrameManager->GetFreeFrame();
@@ -442,6 +437,7 @@ namespace UE::RivermaxCore::Private
 		{
 			return FrameManager->SetFrameData(NewFrameInfo, ReservedFrame);
 		}
+		
 		return false;
 	}
 
@@ -916,7 +912,11 @@ namespace UE::RivermaxCore::Private
 		if (ReservedFrame.IsValid())
 		{
 			ReservedFrame->SetFrameCounter(FrameCounter);
-			ReservedFrames.Enqueue(ReservedFrame);
+			ReservedFrames.Add(FrameCounter, ReservedFrame);
+		}
+		else
+		{
+			check(Options.FrameLockingMode != EFrameLockingMode::BlockOnReservation);
 		}
 
 		return ReservedFrame.IsValid();
diff --git a/Engine/Plugins/VirtualProduction/Rivermax/RivermaxMedia/Source/RivermaxMedia/Private/RivermaxMediaCapture.cpp b/Engine/Plugins/VirtualProduction/Rivermax/RivermaxMedia/Source/RivermaxMedia/Private/RivermaxMediaCapture.cpp
--- a/Engine/Plugins/VirtualProduction/Rivermax/RivermaxMedia/Source/RivermaxMedia/Private/RivermaxMediaCapture.cpp
+++ b/Engine/Plugins/VirtualProduction/Rivermax/RivermaxMedia/Source/RivermaxMedia/Private/RivermaxMediaCapture.cpp
@@ -125,6 +125,27 @@ void URivermaxMediaCapture::GetLastPresentedFrameInformation(UE::RivermaxCore::F
 	}
 }
 
+void URivermaxMediaCapture::WaitForGPU(FRHITexture* InRHITexture)
+{
+	TRACE_CPUPROFILER_EVENT_SCOPE(URivermaxMediaCapture::GPUDirect_WaitForGPU);
+
+	// Since render fence requires us to provide RHICmdList handle we need ot wait on rendering thread. We then add a sync event that we wait for on any thread.
+	FEvent* WaitForRenderCompleteEvent = FPlatformProcess::GetSynchEventFromPool(true /*bIsManualReset*/);
+
+	ENQUEUE_RENDER_COMMAND(WaitForFenceCommand)(
+		[CompleteFence = ShaderCompletedRenderingFence, WaitForRenderCompleteEvent](FRHICommandListImmediate& RHICmdList)
+		{
+			CompleteFence->Wait(RHICmdList, FRHIGPUMask::All());
+			CompleteFence->Clear();
+			WaitForRenderCompleteEvent->Trigger();
+		}
+		);
+
+	WaitForRenderCompleteEvent->Wait();
+
+	FPlatformProcess::ReturnSynchEventToPool(WaitForRenderCompleteEvent);
+}
+
 bool URivermaxMediaCapture::ValidateMediaOutput() const
 {
 	URivermaxMediaOutput* RivermaxMediaOutput = Cast<URivermaxMediaOutput>(MediaOutput);
@@ -223,6 +244,8 @@ bool URivermaxMediaCapture::Initialize(const UE::RivermaxCore::FRivermaxOutputOp
 		}
 	}
 
+	ShaderCompletedRenderingFence = RHICreateGPUFence(TEXT("RmaxRenderingCompleteFence"));
+
 	return bIsActive;
 }
 
@@ -333,7 +356,6 @@ void URivermaxMediaCapture::OnRHIResourceCapturedInternal_AnyThread(const FCaptu
 {
 	using namespace UE::RivermaxCore;
 	TRACE_CPUPROFILER_EVENT_SCOPE_TEXT(*FRivermaxTracingUtils::RmaxOutMediaCapturePipeTraceEvents[InBaseData.SourceFrameNumberRenderThread % 10]);
-
 	TSharedPtr<FRivermaxOutputVideoFrameInfo> NewFrame = MakeShared<FRivermaxOutputVideoFrameInfo>();
 	NewFrame->FrameIdentifier = InBaseData.SourceFrameNumberRenderThread;
 	NewFrame->GPUBuffer = InBuffer;
@@ -438,7 +460,7 @@ void URivermaxMediaCapture::OnCustomCapture_RenderingThread(FRDGBuilder& GraphBu
 	constexpr bool bDoLinearToSRGB = false;
 	const FIntPoint AlignedBufferSize = Options.GetStreamOptions<FRivermaxVideoOutputOptions>(ERivermaxStreamType::VIDEO_2110_20_STREAM)->AlignedResolution;
 	const FIntVector GroupCount = FComputeShaderUtils::GetGroupCount(DesiredOutputSize.X, 64);
-	
+
 	switch (RivermaxOutput->VideoStream.PixelFormat)
 	{
 	case ERivermaxMediaOutputPixelFormat::PF_8BIT_YUV422:
@@ -526,7 +548,20 @@ void URivermaxMediaCapture::OnCustomCapture_RenderingThread(FRDGBuilder& GraphBu
 		break;
 	}
 	}
-	
+
+	//It is only in case of GPU Direct that we need to manually wait for the work to be completed.
+	if (ShouldCaptureRHIResource())
+	{
+		GraphBuilder.AddPass(
+			RDG_EVENT_NAME("RivermaxCaptureGPUFence"),
+			ERDGPassFlags::None,
+			[CompleteFence = ShaderCompletedRenderingFence](FRHICommandList& RHICmdList)
+			{
+				RHICmdList.WriteGPUFence(CompleteFence);
+			}
+		);
+	}
+
 	AddFrameReservationPass(GraphBuilder);
 }
 
diff --git a/Engine/Plugins/VirtualProduction/Rivermax/RivermaxMedia/Source/RivermaxMedia/Public/RivermaxMediaCapture.h b/Engine/Plugins/VirtualProduction/Rivermax/RivermaxMedia/Source/RivermaxMedia/Public/RivermaxMediaCapture.h
--- a/Engine/Plugins/VirtualProduction/Rivermax/RivermaxMedia/Source/RivermaxMedia/Public/RivermaxMediaCapture.h
+++ b/Engine/Plugins/VirtualProduction/Rivermax/RivermaxMedia/Source/RivermaxMedia/Public/RivermaxMediaCapture.h
@@ -29,6 +29,9 @@ public:
 	/** Returns information about last presented frame on the output stream */
 	void GetLastPresentedFrameInformation(UE::RivermaxCore::FPresentedFrameInfo& OutFrameInfo) const;
 
+	/** In cases of GPUDirect we should manually wait for work to be done. */
+	virtual void WaitForGPU(FRHITexture* InRHITexture) override;
+
 public:
 
 	//~ Begin UMediaCapture interface
@@ -104,6 +107,9 @@ private:
 	/** Set of options used to configure output stream */
 	UE::RivermaxCore::FRivermaxOutputOptions Options;
 
+	/** When using GPUDirect we have to rely on our own fence to wait for output to be converted into buffer. This will be waited for in WaitForGPU function. */
+	TRefCountPtr<FRHIGPUFence> ShaderCompletedRenderingFence;
+
 	/** Whether capture is active. Used to be queried by sync task */
 	std::atomic<bool> bIsActive;
 };
