Rivermax: GPUDirect on output fixed a crash caused by repeatedly writing GPU fence without clearing it and subsequent deadlock due to RHI thread being locked by Frame reservation.

#rnx

diff --git a/Engine/Plugins/VirtualProduction/Rivermax/RivermaxMedia/Source/RivermaxMedia/Private/RivermaxMediaCapture.cpp b/Engine/Plugins/VirtualProduction/Rivermax/RivermaxMedia/Source/RivermaxMedia/Private/RivermaxMediaCapture.cpp
--- a/Engine/Plugins/VirtualProduction/Rivermax/RivermaxMedia/Source/RivermaxMedia/Private/RivermaxMediaCapture.cpp
+++ b/Engine/Plugins/VirtualProduction/Rivermax/RivermaxMedia/Source/RivermaxMedia/Private/RivermaxMediaCapture.cpp
@@ -127,23 +127,18 @@ void URivermaxMediaCapture::GetLastPresentedFrameInformation(UE::RivermaxCore::F
 
 void URivermaxMediaCapture::WaitForGPU(FRHITexture* InRHITexture)
 {
-	TRACE_CPUPROFILER_EVENT_SCOPE(URivermaxMediaCapture::GPUDirect_WaitForGPU);
-
-	// Since render fence requires us to provide RHICmdList handle we need ot wait on rendering thread. We then add a sync event that we wait for on any thread.
-	FEvent* WaitForRenderCompleteEvent = FPlatformProcess::GetSynchEventFromPool(true /*bIsManualReset*/);
-
-	ENQUEUE_RENDER_COMMAND(WaitForFenceCommand)(
-		[CompleteFence = ShaderCompletedRenderingFence, WaitForRenderCompleteEvent](FRHICommandListImmediate& RHICmdList)
+	if (bIsActive && ShouldCaptureRHIResource())
+	{
+		TRACE_CPUPROFILER_EVENT_SCOPE(URivermaxMediaCapture::GPUDirect_WaitForGPU);
+		while (!ShaderCompletedRenderingFence->Poll())
 		{
-			CompleteFence->Wait(RHICmdList, FRHIGPUMask::All());
-			CompleteFence->Clear();
-			WaitForRenderCompleteEvent->Trigger();
+			FPlatformProcess::YieldThread();
 		}
-		);
-
-	WaitForRenderCompleteEvent->Wait();
-
-	FPlatformProcess::ReturnSynchEventToPool(WaitForRenderCompleteEvent);
+		ShaderCompletedRenderingFence->Clear();
+		GPUWaitCompleteEvent->Trigger();
+	}
+ 
+ 
 }
 
 bool URivermaxMediaCapture::ValidateMediaOutput() const
@@ -245,6 +240,14 @@ bool URivermaxMediaCapture::Initialize(const UE::RivermaxCore::FRivermaxOutputOp
 	}
 
 	ShaderCompletedRenderingFence = RHICreateGPUFence(TEXT("RmaxRenderingCompleteFence"));
+	GPUWaitCompleteEvent = MakeShareable(FPlatformProcess::GetSynchEventFromPool(true /*bIsManualReset*/), [](FEvent* EventToDelete)
+		{
+			FPlatformProcess::ReturnSynchEventToPool(EventToDelete);
+		});
+
+	// This event starts with a triggered state because Render thread expects this to be triggered to issue render commands.
+	// Since nothing has been completed yet, this needs to be triggered.
+	GPUWaitCompleteEvent->Trigger();
 
 	return bIsActive;
 }
@@ -553,10 +556,15 @@ void URivermaxMediaCapture::OnCustomCapture_RenderingThread(FRDGBuilder& GraphBu
 	if (ShouldCaptureRHIResource())
 	{
 		GraphBuilder.AddPass(
-			RDG_EVENT_NAME("RivermaxCaptureGPUFence"),
-			ERDGPassFlags::None,
-			[CompleteFence = ShaderCompletedRenderingFence](FRHICommandList& RHICmdList)
+			RDG_EVENT_NAME("RivermaxWriteGPUFence"),
+			ERDGPassFlags::NeverCull,
+			[CompleteFence = ShaderCompletedRenderingFence, InGPUWaitCompleteEvent = GPUWaitCompleteEvent](FRHICommandList& RHICmdList)
 			{
+				TRACE_CPUPROFILER_EVENT_SCOPE(RivermaxWriteGPUFence);
+				// Render thread could get to this point before Fence is reset. This could cause a deadlock, therefore it 
+				// is necessary to wait for Sync thread to complete the copy
+				InGPUWaitCompleteEvent->Wait();
+				InGPUWaitCompleteEvent->Reset();
 				RHICmdList.WriteGPUFence(CompleteFence);
 			}
 		);
diff --git a/Engine/Plugins/VirtualProduction/Rivermax/RivermaxMedia/Source/RivermaxMedia/Public/RivermaxMediaCapture.h b/Engine/Plugins/VirtualProduction/Rivermax/RivermaxMedia/Source/RivermaxMedia/Public/RivermaxMediaCapture.h
--- a/Engine/Plugins/VirtualProduction/Rivermax/RivermaxMedia/Source/RivermaxMedia/Public/RivermaxMediaCapture.h
+++ b/Engine/Plugins/VirtualProduction/Rivermax/RivermaxMedia/Source/RivermaxMedia/Public/RivermaxMediaCapture.h
@@ -110,6 +110,9 @@ private:
 	/** When using GPUDirect we have to rely on our own fence to wait for output to be converted into buffer. This will be waited for in WaitForGPU function. */
 	TRefCountPtr<FRHIGPUFence> ShaderCompletedRenderingFence;
 
+	/** Used with GPUDirect An Event used to block from render thread writing a fence again when we wait for it. */
+	TSharedPtr<FEvent, ESPMode::ThreadSafe> GPUWaitCompleteEvent;
+
 	/** Whether capture is active. Used to be queried by sync task */
 	std::atomic<bool> bIsActive;
 };
