Rivermax/ndisplay:
- A workaround for the issue with custom resolutions killing the performance due to Multi-SRD issue in Rivermax.
- Plug for thread safety issue with frame reservation when bIsActive set to false by [REDACTED_PERSON] in Rivermax Stream.

#rnx

diff --git a/Engine/Plugins/Runtime/nDisplay/Source/DisplayClusterMedia/Private/Input/DisplayClusterMediaInputBase.cpp b/Engine/Plugins/Runtime/nDisplay/Source/DisplayClusterMedia/Private/Input/DisplayClusterMediaInputBase.cpp
--- a/Engine/Plugins/Runtime/nDisplay/Source/DisplayClusterMedia/Private/Input/DisplayClusterMediaInputBase.cpp
+++ b/Engine/Plugins/Runtime/nDisplay/Source/DisplayClusterMedia/Private/Input/DisplayClusterMediaInputBase.cpp
@@ -37,12 +37,14 @@ TAutoConsoleVariable<bool> CVarTempRivermaxCropWorkaround(
 	ECVF_RenderThreadSafe
 );
 
-// Based on the discussion, it looks like the problem is the incoming 2110 textures
-// may have up to 3 ExtraPixelsThreshold extra pixels.
+// Due to various alignment requirements the texture coming from Rivermax may be padded.
+// This CVar controls the maximum number of pixels removed in case the received resolution
+// doesn't match expected. #500 roughly corresponds to max number of pixels that can fit into UDP packet.
 TAutoConsoleVariable<int32> CVarTempRivermaxExtraPixelsThreshold(
 	TEXT("nDisplay.Media.Rivermax.ExtraPixelsThreshold"),
-	3,
-	TEXT("nDisplay workaround for Rivermax input\n"),
+	500,
+	TEXT("This sets a threshold for how many pixels will be cropped from the right side of the Rivermax stream \
+		if the stream was padded. \n"),
 	ECVF_RenderThreadSafe
 );
 
diff --git a/Engine/Plugins/VirtualProduction/Rivermax/RivermaxCore/Source/RivermaxCore/Private/RivermaxUtils.h b/Engine/Plugins/VirtualProduction/Rivermax/RivermaxCore/Source/RivermaxCore/Private/RivermaxUtils.h
--- a/Engine/Plugins/VirtualProduction/Rivermax/RivermaxCore/Source/RivermaxCore/Private/RivermaxUtils.h
+++ b/Engine/Plugins/VirtualProduction/Rivermax/RivermaxCore/Source/RivermaxCore/Private/RivermaxUtils.h
@@ -19,12 +19,6 @@ namespace UE::RivermaxCore::Private::Utils
 	static constexpr uint32 FullHDHeight = 1080;
 	static constexpr uint32 FullHDWidth = 1920;
 
-	/** Maximum payload size in bytes that the plugin can send based on UDP max size and RTP header.  */
-	static constexpr uint32 MaxPayloadSize = 1420;
-	
-	/** Smallest payload size (bytes) to use as a lower bound in search for a payload that can be equal across a line */
-	static constexpr uint32 MinPayloadSize = 600;
-
 	/** SMTPE 2110-10.The Media Clock and RTP Clock rate for streams compliant to this standard shall be 90 kHz. */
 	static constexpr double MediaClockSampleRate = 90000.0;
 	
diff --git a/Engine/Plugins/VirtualProduction/Rivermax/RivermaxCore/Source/RivermaxCore/Private/Streams/RivermaxOutStream.cpp b/Engine/Plugins/VirtualProduction/Rivermax/RivermaxCore/Source/RivermaxCore/Private/Streams/RivermaxOutStream.cpp
--- a/Engine/Plugins/VirtualProduction/Rivermax/RivermaxCore/Source/RivermaxCore/Private/Streams/RivermaxOutStream.cpp
+++ b/Engine/Plugins/VirtualProduction/Rivermax/RivermaxCore/Source/RivermaxCore/Private/Streams/RivermaxOutStream.cpp
@@ -255,6 +255,8 @@ namespace UE::RivermaxCore::Private
 		{
 			Stop();
 			
+			CurrentFrame.Reset();
+
 			FrameAvailableSignal->Trigger();
 			FrameReadyToSendSignal->Trigger();
 			RivermaxThread->Kill(true);
diff --git a/Engine/Plugins/VirtualProduction/Rivermax/RivermaxCore/Source/RivermaxCore/Private/Streams/RivermaxOutVideoStream.cpp b/Engine/Plugins/VirtualProduction/Rivermax/RivermaxCore/Source/RivermaxCore/Private/Streams/RivermaxOutVideoStream.cpp
--- a/Engine/Plugins/VirtualProduction/Rivermax/RivermaxCore/Source/RivermaxCore/Private/Streams/RivermaxOutVideoStream.cpp
+++ b/Engine/Plugins/VirtualProduction/Rivermax/RivermaxCore/Source/RivermaxCore/Private/Streams/RivermaxOutVideoStream.cpp
@@ -20,7 +20,7 @@
 namespace UE::RivermaxCore::Private
 {
 	static TAutoConsoleVariable<int32> CVarRivermaxOutputEnableMultiSRD(
-		TEXT("Rivermax.Output.EnableMultiSRD"), 1,
+		TEXT("Rivermax.Output.EnableMultiSRD"), 0,
 		TEXT("When enabled and if the row cannot be split evenly, non-uniform payloads will be used. The last packet for the frame will not be fully filled with data.\n"
 			"If disabled, the payloads will be split evenly or the 2110 stream will be disabled."),
 		ECVF_Default);
@@ -52,9 +52,14 @@ namespace UE::RivermaxCore::Private
 		TEXT("Number of chunks between each memcopy to help with timing for different frame format."),
 		ECVF_Default);
 
+	// TODO: To be removed or improved in 5.7. In 5.6.1 Media Capture will be responsible for padding the resolution in case row cannot be split into equal sized packets due to issues with Multi-srd.
 	bool FindPayloadSize(uint32 InBytesPerLine, const uint32 PixelGroupSize, uint16& OutPayloadSize)
 	{
-		using namespace UE::RivermaxCore::Private::Utils;
+		/** Maximum payload size in bytes that the plugin can send based on UDP max size and RTP header.  */
+		static constexpr uint32 MaxPayloadSize = 1420;
+
+		/** Smallest payload size (bytes) to use as a lower bound in search for a payload that can be equal across a line */
+		static constexpr uint32 MinPayloadSize = 600;
 
 		// For now, we only find a payload size that can be equal across one line
 		// Support for last payload of a line being smaller is there but is causing issue
@@ -418,6 +423,11 @@ namespace UE::RivermaxCore::Private
 
 	bool FRivermaxOutVideoStream::PushFrame(TSharedPtr<IRivermaxOutputInfo> FrameInfo)
 	{
+		if (!bIsActive)
+		{
+			return false;
+		}
+
 		TRACE_CPUPROFILER_EVENT_SCOPE(FRivermaxOutVideoStream::PushVideoFrame);
 		TSharedPtr<FRivermaxOutputVideoFrameInfo> NewFrameInfo = StaticCastSharedPtr<FRivermaxOutputVideoFrameInfo>(FrameInfo);
 		NewFrameInfo->Stride = GetRowSizeInBytes();
@@ -897,6 +907,11 @@ namespace UE::RivermaxCore::Private
 
 	bool FRivermaxOutVideoStream::ReserveFrame(uint64 FrameCounter) const
 	{
+		if (!bIsActive)
+		{
+			return false;
+		}
+
 		// There is only one reserved frame at the time per stream.
 		TSharedPtr<FRivermaxOutputFrame> ReservedFrame = FrameManager->GetFreeFrame();
 		if (!ReservedFrame.IsValid() && Options.FrameLockingMode == EFrameLockingMode::BlockOnReservation)
@@ -909,15 +924,11 @@ namespace UE::RivermaxCore::Private
 			}
 		}
 
-		if (ReservedFrame.IsValid())
+		if (ReservedFrame.IsValid() && bIsActive)
 		{
 			ReservedFrame->SetFrameCounter(FrameCounter);
 			ReservedFrames.Add(FrameCounter, ReservedFrame);
 		}
-		else
-		{
-			check(Options.FrameLockingMode != EFrameLockingMode::BlockOnReservation);
-		}
 
 		return ReservedFrame.IsValid();
 	}
diff --git a/Engine/Plugins/VirtualProduction/Rivermax/RivermaxMedia/Source/RivermaxMedia/Private/RivermaxMediaPlayer.cpp b/Engine/Plugins/VirtualProduction/Rivermax/RivermaxMedia/Source/RivermaxMedia/Private/RivermaxMediaPlayer.cpp
--- a/Engine/Plugins/VirtualProduction/Rivermax/RivermaxMedia/Source/RivermaxMedia/Private/RivermaxMediaPlayer.cpp
+++ b/Engine/Plugins/VirtualProduction/Rivermax/RivermaxMedia/Source/RivermaxMedia/Private/RivermaxMediaPlayer.cpp
@@ -655,7 +655,8 @@ namespace UE::RivermaxMedia
 			}
 
 			const FIntPoint ProcessedOutputDimension = { (int32)SourceBufferDesc.NumberOfElements, 1 };
-			const FIntVector GroupCount = FComputeShaderUtils::GetGroupCount(SourceBufferDesc.NumberOfElements, 64);
+			const FIntVector GroupCountUnclamped = FComputeShaderUtils::GetGroupCount(SourceBufferDesc.NumberOfElements, 64);
+			const FIntVector GroupCount = GroupCountUnclamped.ComponentMin(GRHIGlobals.MaxDispatchThreadGroupsPerDimension);
 			FGlobalShaderMap* GlobalShaderMap = GetGlobalShaderMap(GMaxRHIFeatureLevel);
 
 			//Configure shader and add conversion pass based on desired pixel format
diff --git a/Engine/Plugins/VirtualProduction/Rivermax/RivermaxMedia/Source/RivermaxMedia/Private/RivermaxMediaUtils.cpp b/Engine/Plugins/VirtualProduction/Rivermax/RivermaxMedia/Source/RivermaxMedia/Private/RivermaxMediaUtils.cpp
--- a/Engine/Plugins/VirtualProduction/Rivermax/RivermaxMedia/Source/RivermaxMedia/Private/RivermaxMediaUtils.cpp
+++ b/Engine/Plugins/VirtualProduction/Rivermax/RivermaxMedia/Source/RivermaxMedia/Private/RivermaxMediaUtils.cpp
@@ -161,7 +161,29 @@ namespace UE::RivermaxMediaUtils::Private
 		const uint32 PixelAlignment = InFormatInfo.PixelGroupCoverage;
 		uint32 PixelsToBeFitIntoGroup = ResolutionToAlign.X % PixelAlignment;
 		const uint32 AlignedHorizontalResolution = PixelsToBeFitIntoGroup > 0 ? ResolutionToAlign.X + (PixelAlignment - PixelsToBeFitIntoGroup) : ResolutionToAlign.X;
-		return FIntPoint(AlignedHorizontalResolution, ResolutionToAlign.Y);
+
+		// This is a workaround for an issue with Multi-SRD packets. We pad to the packet with the smallest requirement for padding 
+		// so that divides the pixel row into equal sizes. To be removed when Multi-SRD is fixed.
+		uint32 AlignedHorizontalResolutionPacketSRD = AlignedHorizontalResolution;
+
+		IConsoleVariable* CvarEnableMultiSrd = IConsoleManager::Get().FindConsoleVariable(TEXT("Rivermax.Output.EnableMultiSRD"));
+
+		if (CvarEnableMultiSrd && CvarEnableMultiSrd->GetInt() == 0)
+		{
+			const int32 BytesPerRow = (AlignedHorizontalResolution  * InFormatInfo.PixelGroupSize) / InFormatInfo.PixelGroupCoverage;
+
+			FPayloadSizeInformation PayloadSizeInfo;
+			FindPayloadSize(BytesPerRow, InFormatInfo.PixelGroupSize, PayloadSizeInfo);
+
+			// Pad if the appropriate payload size isn't found.
+			if (!PayloadSizeInfo.bIsExactFit)
+			{
+				const uint32 BytesPerRowAligned = ((BytesPerRow + PayloadSizeInfo.PayloadSize - 1) / PayloadSizeInfo.PayloadSize) * PayloadSizeInfo.PayloadSize;
+				AlignedHorizontalResolutionPacketSRD = (BytesPerRowAligned * InFormatInfo.PixelGroupCoverage)  / InFormatInfo.PixelGroupSize;
+			}
+		}
+
+		return FIntPoint(AlignedHorizontalResolutionPacketSRD, ResolutionToAlign.Y);
 	}
 
 	UE::RivermaxCore::ERivermaxAlignmentMode MediaOutputAlignmentToRivermaxAlignment(ERivermaxMediaAlignmentMode InAlignmentMode)
@@ -482,6 +504,71 @@ namespace UE::RivermaxMediaUtils::Private
 		return true;
 	}
 
+	void FindPayloadSize(uint32 InBytesPerLine, const uint32 PixelGroupSize, FPayloadSizeInformation& OutPayloadSizeInformation)
+	{
+		OutPayloadSizeInformation.bIsExactFit = false;
+
+		// For now, we only find a payload size that can be equal across one line
+		// Support for last payload of a line being smaller is there but is causing issue
+		// We fail to output if we receive a resolution for which we can't find an equal payload size
+		int32 NumberOfPacketsPerLine = InBytesPerLine / MaxPayloadSize;
+		if (NumberOfPacketsPerLine == 0)
+		{
+			if (InBytesPerLine > MinPayloadSize && InBytesPerLine % PixelGroupSize == 0)
+			{
+				OutPayloadSizeInformation.PayloadSize = InBytesPerLine;
+				OutPayloadSizeInformation.bIsExactFit = true;
+				return;
+			}
+
+			// Pad to the smallest packet size since the resolution is too small.
+			OutPayloadSizeInformation.PayloadSize = FMath::Max(MinPayloadSize, InBytesPerLine);
+			if (InBytesPerLine % OutPayloadSizeInformation.PayloadSize != 0)
+			{
+				OutPayloadSizeInformation.PayloadSize = OutPayloadSizeInformation.PayloadSize + PixelGroupSize - (OutPayloadSizeInformation.PayloadSize % PixelGroupSize);
+			}
+			OutPayloadSizeInformation.bIsExactFit = InBytesPerLine % OutPayloadSizeInformation.PayloadSize == 0;
+			return;
+		}
+
+		// By default the maximum udp size aligned to the nearest pixel group size is used. Then a more appropriate payload is attempted to be found.
+		uint32 ClosestPayloadSizeToPadTo = MaxPayloadSize - (MaxPayloadSize % PixelGroupSize);
+
+		while (true)
+		{
+			const int32 TestPacketSize = InBytesPerLine / NumberOfPacketsPerLine;
+			if (TestPacketSize < MinPayloadSize)
+			{
+				break;
+			}
+
+			if (TestPacketSize <= MaxPayloadSize)
+			{
+				if ((TestPacketSize % PixelGroupSize) == 0)
+				{
+					// This means we have found a packet size that requires a smaller number of pixels to pad to.
+					if (InBytesPerLine % TestPacketSize < (InBytesPerLine % ClosestPayloadSizeToPadTo))
+					{
+						ClosestPayloadSizeToPadTo = TestPacketSize;
+					}
+
+					if ((InBytesPerLine % NumberOfPacketsPerLine) == 0)
+					{
+						OutPayloadSizeInformation.PayloadSize = TestPacketSize;
+						OutPayloadSizeInformation.bIsExactFit = true;
+						return;
+					}
+				}
+			}
+
+			++NumberOfPacketsPerLine;
+		}
+
+		OutPayloadSizeInformation.PayloadSize = ClosestPayloadSizeToPadTo;
+		OutPayloadSizeInformation.bIsExactFit = false;
+		return;
+	}
+
 	UE::RivermaxCore::EFrameLockingMode MediaOutputFrameLockingToRivermax(ERivermaxFrameLockingMode InFrameLockingMode)
 	{
 		using namespace UE::RivermaxCore;
diff --git a/Engine/Plugins/VirtualProduction/Rivermax/RivermaxMedia/Source/RivermaxMedia/Private/RivermaxMediaUtils.h b/Engine/Plugins/VirtualProduction/Rivermax/RivermaxMedia/Source/RivermaxMedia/Private/RivermaxMediaUtils.h
--- a/Engine/Plugins/VirtualProduction/Rivermax/RivermaxMedia/Source/RivermaxMedia/Private/RivermaxMediaUtils.h
+++ b/Engine/Plugins/VirtualProduction/Rivermax/RivermaxMedia/Source/RivermaxMedia/Private/RivermaxMediaUtils.h
@@ -19,6 +19,28 @@ namespace UE::RivermaxMediaUtils::Private
 		uint32 NumberOfElements = 0;
 	};
 
+	struct FPayloadSizeInformation
+	{
+		/** True if an exact fit is found. I.e. a packet size was found that divides sample row without the remainder. */
+		bool bIsExactFit = false;
+		
+		/**
+		* If an exact fit this PayloadSize will be used without padding the resolution.
+		* If not an exact fit then the PayloadSize will be used to pad the resolution to.
+		*/
+		uint32 PayloadSize = 0;
+	};
+
+	/** Maximum payload size in bytes that the plugin can send based on UDP max size and RTP header.  */
+	static constexpr uint32 MaxPayloadSize = 1420;
+
+	/** 
+	* Smallest payload size (bytes) to use as a lower bound in search for a payload that can be equal across a line.
+	* The smaller the packet size, the greater the overhead. 
+	* More packets mean more RTP and UDP headers need to be included, and more processing is required as a result.
+	*/
+	static constexpr uint32 MinPayloadSize = 600;
+
 	UE::RivermaxCore::ESamplingType MediaOutputPixelFormatToRivermaxSamplingType(ERivermaxMediaOutputPixelFormat InPixelFormat);
 	UE::RivermaxCore::ESamplingType MediaSourcePixelFormatToRivermaxSamplingType(ERivermaxMediaSourcePixelFormat InPixelFormat);
 	ERivermaxMediaSourcePixelFormat RivermaxPixelFormatToMediaSourcePixelFormat(UE::RivermaxCore::ESamplingType InSamplingType);
@@ -35,6 +57,13 @@ namespace UE::RivermaxMediaUtils::Private
 	* Returns true if SDP could successfully be created from rivermax options.
 	*/
 	bool StreamOptionsToSDPDescription(const UE::RivermaxCore::FRivermaxOutputOptions& Options, TArray<char>& OutSDPDescription);
+
+	/**
+	* Find the right size for the data portion of the UDP packet. It is generally desired to maximize the size of the packet and make it as close
+	* possible to the maximum UDP packet size which is typically 1500 with the header.
+	* OutPayloadSizeInformation contains information about the packet that requires smallest padding if an exact packet that fits equal number of bytes isn't found.
+	*/
+	void FindPayloadSize(uint32 InBytesPerLine, const uint32 PixelGroupSize, FPayloadSizeInformation& OutPayloadSizeInformation);
 }
 
 // Custom version to keep track of and restore depricated properties.
