Rivermax: A "fix" for static analysis warning: Mod by zero. Denominator range [0..1420]

#rnx

diff --git a/Engine/Plugins/VirtualProduction/Rivermax/RivermaxMedia/Source/RivermaxMedia/Private/RivermaxMediaUtils.cpp b/Engine/Plugins/VirtualProduction/Rivermax/RivermaxMedia/Source/RivermaxMedia/Private/RivermaxMediaUtils.cpp
--- a/Engine/Plugins/VirtualProduction/Rivermax/RivermaxMedia/Source/RivermaxMedia/Private/RivermaxMediaUtils.cpp
+++ b/Engine/Plugins/VirtualProduction/Rivermax/RivermaxMedia/Source/RivermaxMedia/Private/RivermaxMediaUtils.cpp
@@ -170,7 +170,8 @@ namespace UE::RivermaxMediaUtils::Private
 
 		if (CvarEnableMultiSrd && CvarEnableMultiSrd->GetInt() == 0)
 		{
-			const int32 BytesPerRow = (AlignedHorizontalResolution  * InFormatInfo.PixelGroupSize) / InFormatInfo.PixelGroupCoverage;
+			// PixelGroupSize is always exactly divisible by PixelGroupCoverage;
+			const uint32 BytesPerRow = (AlignedHorizontalResolution  * InFormatInfo.PixelGroupSize) / InFormatInfo.PixelGroupCoverage;
 
 			FPayloadSizeInformation PayloadSizeInfo;
 			FindPayloadSize(BytesPerRow, InFormatInfo.PixelGroupSize, PayloadSizeInfo);
@@ -525,6 +526,8 @@ namespace UE::RivermaxMediaUtils::Private
 			OutPayloadSizeInformation.PayloadSize = FMath::Max(MinPayloadSize, InBytesPerLine);
 			if (InBytesPerLine % OutPayloadSizeInformation.PayloadSize != 0)
 			{
+				// This may exceed MaxPayloadSize, which is acceptable since this possibility is accounted for  
+				// by setting MaxPayloadSize to be a few pixel groups' worth of bytes smaller.
 				OutPayloadSizeInformation.PayloadSize = OutPayloadSizeInformation.PayloadSize + PixelGroupSize - (OutPayloadSizeInformation.PayloadSize % PixelGroupSize);
 			}
 			OutPayloadSizeInformation.bIsExactFit = InBytesPerLine % OutPayloadSizeInformation.PayloadSize == 0;
@@ -532,11 +535,11 @@ namespace UE::RivermaxMediaUtils::Private
 		}
 
 		// By default the maximum udp size aligned to the nearest pixel group size is used. Then a more appropriate payload is attempted to be found.
-		uint32 ClosestPayloadSizeToPadTo = MaxPayloadSize - (MaxPayloadSize % PixelGroupSize);
+		uint32 ClosestPayloadSizeToPadTo = FMath::Max(MaxPayloadSize - (MaxPayloadSize % PixelGroupSize), MinPayloadSize);
 
 		while (true)
 		{
-			const int32 TestPacketSize = InBytesPerLine / NumberOfPacketsPerLine;
+			const uint32 TestPacketSize = InBytesPerLine / NumberOfPacketsPerLine;
 			if (TestPacketSize < MinPayloadSize)
 			{
 				break;
@@ -549,7 +552,7 @@ namespace UE::RivermaxMediaUtils::Private
 					// This means we have found a packet size that requires a smaller number of pixels to pad to.
 					if (InBytesPerLine % TestPacketSize < (InBytesPerLine % ClosestPayloadSizeToPadTo))
 					{
-						ClosestPayloadSizeToPadTo = TestPacketSize;
+						ClosestPayloadSizeToPadTo = FMath::Max(TestPacketSize, MinPayloadSize);
 					}
 
 					if ((InBytesPerLine % NumberOfPacketsPerLine) == 0)
