[Rivermax] Fixed capture size returned by URivermaxMediaOutput
- URivermaxMediaOutput::GetRequestedSize now returns correct capture size or default {0, 0} if not set

#rnx

diff --git a/Engine/Plugins/VirtualProduction/Rivermax/RivermaxMedia/Source/RivermaxMedia/Private/RivermaxMediaOutput.cpp b/Engine/Plugins/VirtualProduction/Rivermax/RivermaxMedia/Source/RivermaxMedia/Private/RivermaxMediaOutput.cpp
--- a/Engine/Plugins/VirtualProduction/Rivermax/RivermaxMedia/Source/RivermaxMedia/Private/RivermaxMediaOutput.cpp
+++ b/Engine/Plugins/VirtualProduction/Rivermax/RivermaxMedia/Source/RivermaxMedia/Private/RivermaxMediaOutput.cpp
@@ -24,9 +24,9 @@ bool URivermaxMediaOutput::Validate(FString& OutFailureReason) const
 
 FIntPoint URivermaxMediaOutput::GetRequestedSize() const
 {
-	if (VideoStream.bOverrideResolution)
+	if (bOverrideResolution)
 	{
-		return VideoStream.Resolution;
+		return Resolution;
 	}
 
 	return UMediaOutput::RequestCaptureSourceSize;
