Media I/O: Avoid raw this pointer capture in async task, which could cause crashes if the texture sample is destroyed before task execution.

diff --git a/Engine/Plugins/Media/MediaIOFramework/Source/MediaIOCore/Private/Shared/MediaIOCoreTextureSampleBase.cpp b/Engine/Plugins/Media/MediaIOFramework/Source/MediaIOCore/Private/Shared/MediaIOCoreTextureSampleBase.cpp
--- a/Engine/Plugins/Media/MediaIOFramework/Source/MediaIOCore/Private/Shared/MediaIOCoreTextureSampleBase.cpp
+++ b/Engine/Plugins/Media/MediaIOFramework/Source/MediaIOCore/Private/Shared/MediaIOCoreTextureSampleBase.cpp
@@ -161,9 +161,15 @@ void FMediaIOCoreTextureSampleBase::SetColorConversionSettings(TSharedPtr<struct
 	}
 	else
 	{
-		AsyncTask(ENamedThreads::GameThread, [this]() {
-			CacheColorCoversionSettings_GameThread();
-			});
+		AsyncTask(ENamedThreads::GameThread, [WeakThis = AsWeak()]()
+			{
+				TSharedPtr<FMediaIOCoreTextureSampleBase> This = WeakThis.Pin();
+				if (This.IsValid())
+				{
+					This->CacheColorCoversionSettings_GameThread();
+				}
+			}
+		);
 	}
 }
 
diff --git a/Engine/Plugins/Media/MediaIOFramework/Source/MediaIOCore/Public/MediaIOCoreTextureSampleBase.h b/Engine/Plugins/Media/MediaIOFramework/Source/MediaIOCore/Public/MediaIOCoreTextureSampleBase.h
--- a/Engine/Plugins/Media/MediaIOFramework/Source/MediaIOCore/Public/MediaIOCoreTextureSampleBase.h
+++ b/Engine/Plugins/Media/MediaIOFramework/Source/MediaIOCore/Public/MediaIOCoreTextureSampleBase.h
@@ -92,6 +92,7 @@ class FMediaIOCoreTextureSampleBase
 	: public IMediaTextureSample
 	, public IMediaPoolable
 	, public IMediaTextureSampleColorConverter
+	, public TSharedFromThis<FMediaIOCoreTextureSampleBase, ESPMode::ThreadSafe>
 {
 
 public:
diff --git a/Engine/Plugins/VirtualProduction/Rivermax/RivermaxMedia/Source/RivermaxMedia/Private/RivermaxMediaTextureSample.cpp b/Engine/Plugins/VirtualProduction/Rivermax/RivermaxMedia/Source/RivermaxMedia/Private/RivermaxMediaTextureSample.cpp
--- a/Engine/Plugins/VirtualProduction/Rivermax/RivermaxMedia/Source/RivermaxMedia/Private/RivermaxMediaTextureSample.cpp
+++ b/Engine/Plugins/VirtualProduction/Rivermax/RivermaxMedia/Source/RivermaxMedia/Private/RivermaxMediaTextureSample.cpp
@@ -64,7 +64,7 @@ void FRivermaxMediaTextureSample::InitializeGPUBuffer(const FIntPoint& InResolut
 
 	InputFormat = InSampleFormat;
 	ENQUEUE_RENDER_COMMAND(FRivermaxMediaTextureSample)(
-	[SharedThis = AsShared(), RDGBufferDesc](FRHICommandListImmediate& CommandList)
+	[SharedThis = StaticCastSharedRef<FRivermaxMediaTextureSample>(AsShared()), RDGBufferDesc](FRHICommandListImmediate& CommandList)
 	{
 		SharedThis->GPUBuffer = AllocatePooledBuffer(RDGBufferDesc, TEXT("RmaxInput Buffer"));
 
diff --git a/Engine/Plugins/VirtualProduction/Rivermax/RivermaxMedia/Source/RivermaxMedia/Private/RivermaxMediaTextureSample.h b/Engine/Plugins/VirtualProduction/Rivermax/RivermaxMedia/Source/RivermaxMedia/Private/RivermaxMediaTextureSample.h
--- a/Engine/Plugins/VirtualProduction/Rivermax/RivermaxMedia/Source/RivermaxMedia/Private/RivermaxMediaTextureSample.h
+++ b/Engine/Plugins/VirtualProduction/Rivermax/RivermaxMedia/Source/RivermaxMedia/Private/RivermaxMediaTextureSample.h
@@ -19,7 +19,7 @@ namespace UE::RivermaxMedia
 	/**
 	 * Implements a media texture sample for RivermaxMediaPlayer.  
 	 */
-	class FRivermaxMediaTextureSample : public FMediaIOCoreTextureSampleBase, public UE::RivermaxCore::IRivermaxVideoSample, public TSharedFromThis<FRivermaxMediaTextureSample>
+	class FRivermaxMediaTextureSample : public FMediaIOCoreTextureSampleBase, public UE::RivermaxCore::IRivermaxVideoSample
 	{
 	public:
 		FRivermaxMediaTextureSample();
