MediaIO: Fix incorrectly removed previous Rivermax defaults in CL 47872781.

#Synced-CL 49805348

diff --git a/Engine/Plugins/VirtualProduction/Rivermax/RivermaxMedia/Source/RivermaxMedia/Private/RivermaxMediaSource.cpp b/Engine/Plugins/VirtualProduction/Rivermax/RivermaxMedia/Source/RivermaxMedia/Private/RivermaxMediaSource.cpp
--- a/Engine/Plugins/VirtualProduction/Rivermax/RivermaxMedia/Source/RivermaxMedia/Private/RivermaxMediaSource.cpp
+++ b/Engine/Plugins/VirtualProduction/Rivermax/RivermaxMedia/Source/RivermaxMedia/Private/RivermaxMediaSource.cpp
@@ -16,6 +16,12 @@ URivermaxMediaSource::URivermaxMediaSource() : UCaptureCardMediaSource()
 	Deinterlacer = nullptr;
 	bRenderJIT = true;
 	EvaluationType = EMediaIOSampleEvaluationType::Latest;
+#if WITH_EDITORONLY_DATA
+PRAGMA_DISABLE_DEPRECATION_WARNINGS
+	bOverrideSourceEncoding_DEPRECATED = false;
+	bOverrideSourceColorSpace_DEPRECATED = false;
+PRAGMA_ENABLE_DEPRECATION_WARNINGS
+#endif
 }
 
 bool URivermaxMediaSource::GetMediaOption(const FName& Key, bool DefaultValue) const
